Docs: Removed file comments and added missing code documentation

diff --git a/Engine/Source/Editor/UndoHistory/Private/Widgets/SUndoHistoryTableRow.h b/Engine/Source/Editor/UndoHistory/Private/Widgets/SUndoHistoryTableRow.h
--- a/Engine/Source/Editor/UndoHistory/Private/Widgets/SUndoHistoryTableRow.h
+++ b/Engine/Source/Editor/UndoHistory/Private/Widgets/SUndoHistoryTableRow.h
@@ -1,9 +1,5 @@
 // Copyright 1998-2013 Epic Games, Inc. All Rights Reserved.
 
-/*=============================================================================
-	SUndoHistoryTableRow.h: Declares the SUndoHistoryTableRow class.
-=============================================================================*/
-
 #pragma once
 
 
@@ -29,7 +25,7 @@ public:
 	/**
 	 * Constructs the widget.
 	 *
-	 * @param InArgs - The construction arguments.
+	 * @param InArgs The construction arguments.
 	 */
 	void Construct( const FArguments& InArgs, const TSharedRef<STableViewBase>& InOwnerTableView )
 	{
diff --git a/Engine/Source/Editor/UndoHistory/Public/Interfaces/IUndoHistoryModule.h b/Engine/Source/Editor/UndoHistory/Public/Interfaces/IUndoHistoryModule.h
--- a/Engine/Source/Editor/UndoHistory/Public/Interfaces/IUndoHistoryModule.h
+++ b/Engine/Source/Editor/UndoHistory/Public/Interfaces/IUndoHistoryModule.h
@@ -1,9 +1,5 @@
 // Copyright 1998-2013 Epic Games, Inc. All Rights Reserved.
 
-/*=============================================================================
-	IUndoHistoryModule.h: Declares the IUndoHistoryModule interface.
-=============================================================================*/
-
 #pragma once
 
 
@@ -18,8 +14,7 @@ public:
 	/**
 	 * Gets a reference to the UndoHistory module instance.
 	 *
-	 * @todo gmp - better implementation using dependency injection.
-	 *
+	 * @todo gmp: better implementation using dependency injection.
 	 * @return A reference to the module.
 	 */
 	static IUndoHistoryModule& Get( )
@@ -27,7 +22,6 @@ public:
 		return FModuleManager::LoadModuleChecked<IUndoHistoryModule>("UndoHistory");
 	}
 
-
 public:
 
 	/**
diff --git a/Engine/Source/Editor/UndoHistory/Public/UndoHistory.h b/Engine/Source/Editor/UndoHistory/Public/UndoHistory.h
--- a/Engine/Source/Editor/UndoHistory/Public/UndoHistory.h
+++ b/Engine/Source/Editor/UndoHistory/Public/UndoHistory.h
@@ -1,8 +1,5 @@
 // Copyright 1998-2013 Epic Games, Inc. All Rights Reserved.
 
-/*=============================================================================
-	UndoHistory.h: UndoHistory module public header file.
-=============================================================================*/
 
 #pragma once
 
diff --git a/Engine/Source/Programs/UnrealFrontend/Private/Mac/MacUnrealFrontendMain.cpp b/Engine/Source/Programs/UnrealFrontend/Private/Mac/MacUnrealFrontendMain.cpp
--- a/Engine/Source/Programs/UnrealFrontend/Private/Mac/MacUnrealFrontendMain.cpp
+++ b/Engine/Source/Programs/UnrealFrontend/Private/Mac/MacUnrealFrontendMain.cpp
@@ -1,9 +1,5 @@
 // Copyright 1998-2014 Epic Games, Inc. All Rights Reserved.
 
-/*=============================================================================
-	MacUnrealFrontendMain.mm: Implements the main entry point for MacOS.
-=============================================================================*/
-
 #include "UnrealFrontendMain.h"
 
 
diff --git a/Engine/Source/Programs/UnrealFrontend/Private/UnrealFrontendMain.cpp b/Engine/Source/Programs/UnrealFrontend/Private/UnrealFrontendMain.cpp
--- a/Engine/Source/Programs/UnrealFrontend/Private/UnrealFrontendMain.cpp
+++ b/Engine/Source/Programs/UnrealFrontend/Private/UnrealFrontendMain.cpp
@@ -1,9 +1,5 @@
 // Copyright 1998-2014 Epic Games, Inc. All Rights Reserved.
 
-/*=============================================================================
-	UnrealFrontendMain.cpp: Implements the UnrealFrontend application's main loop.
-=============================================================================*/
-
 #include "UnrealFrontendMain.h"
 
 #include "RequiredProgramMainCPPInclude.h"
@@ -21,6 +17,7 @@ IMPLEMENT_APPLICATION(UnrealFrontend, "UnrealFrontend");
 
 #define IDEAL_FRAMERATE 60;
 
+
 namespace WorkspaceMenu
 {
 	TSharedRef<FWorkspaceItem> DeveloperTools = FWorkspaceItem::NewGroup( NSLOCTEXT("UnrealFrontend", "DeveloperToolsMenu", "Developer Tools") );
diff --git a/Engine/Source/Programs/UnrealFrontend/Private/UnrealFrontendMain.h b/Engine/Source/Programs/UnrealFrontend/Private/UnrealFrontendMain.h
--- a/Engine/Source/Programs/UnrealFrontend/Private/UnrealFrontendMain.h
+++ b/Engine/Source/Programs/UnrealFrontend/Private/UnrealFrontendMain.h
@@ -1,9 +1,5 @@
 // Copyright 1998-2014 Epic Games, Inc. All Rights Reserved.
 
-/*=============================================================================
-	UnrealFrontendMain.h: Declares the UnrealFrontend application's main loop.
-=============================================================================*/
-
 #pragma once
 
 #include "Core.h"
@@ -15,8 +11,7 @@
 /**
  * The application's main function.
  *
- * @param CommandLine - The application command line.
- *
+ * @param CommandLine The application command line.
  * @return Application's exit value.
  */
 int32 UnrealFrontendMain( const TCHAR* CommandLine );
diff --git a/Engine/Source/Programs/UnrealFrontend/Private/Windows/WindowsUnrealFrontendMain.cpp b/Engine/Source/Programs/UnrealFrontend/Private/Windows/WindowsUnrealFrontendMain.cpp
--- a/Engine/Source/Programs/UnrealFrontend/Private/Windows/WindowsUnrealFrontendMain.cpp
+++ b/Engine/Source/Programs/UnrealFrontend/Private/Windows/WindowsUnrealFrontendMain.cpp
@@ -1,21 +1,17 @@
 // Copyright 1998-2014 Epic Games, Inc. All Rights Reserved.
 
-/*=============================================================================
-	WindowsUnrealFrontendMain.cpp: Implements the main entry point for Windows.
-=============================================================================*/
-
 #include "UnrealFrontendMain.h"
 #include "ExceptionHandling.h"
 #include "LaunchEngineLoop.h"
 
+
 /**
  * The main application entry point for Windows platforms.
  *
- * @param hInInstance - Handle to the current instance of the application.
- * @param hPrevInstance - Handle to the previous instance of the application (always NULL).
- * @param lpCmdLine - Command line for the application.
- * @param nShowCmd - Specifies how the window is to be shown.
- *
+ * @param hInInstance Handle to the current instance of the application.
+ * @param hPrevInstance Handle to the previous instance of the application (always NULL).
+ * @param lpCmdLine Command line for the application.
+ * @param nShowCmd Specifies how the window is to be shown.
  * @return Application's exit value.
  */
 int32 WINAPI WinMain( HINSTANCE hInInstance, HINSTANCE hPrevInstance, char* lpCmdLine, int32 nShowCmd )
diff --git a/Engine/Source/Programs/UnrealHeaderTool/Private/BaseParser.h b/Engine/Source/Programs/UnrealHeaderTool/Private/BaseParser.h
--- a/Engine/Source/Programs/UnrealHeaderTool/Private/BaseParser.h
+++ b/Engine/Source/Programs/UnrealHeaderTool/Private/BaseParser.h
@@ -1,9 +1,5 @@
 // Copyright 1998-2014 Epic Games, Inc. All Rights Reserved.
 
-/*=============================================================================
-	HeaderParser.h: Parses annotated C++ headers to generate additional code and metadata.
-=============================================================================*/
-
 #pragma once
 
 #include "ParserHelper.h"
diff --git a/Engine/Source/Programs/UnrealHeaderTool/Private/Class.h b/Engine/Source/Programs/UnrealHeaderTool/Private/Class.h
--- a/Engine/Source/Programs/UnrealHeaderTool/Private/Class.h
+++ b/Engine/Source/Programs/UnrealHeaderTool/Private/Class.h
@@ -1,9 +1,5 @@
 // Copyright 1998-2014 Epic Games, Inc. All Rights Reserved.
 
-/*=============================================================================
-	Class.h: Represents a parsed class.
-=============================================================================*/
-
 #pragma once
 
 struct EEnforceInterfacePrefix
diff --git a/Engine/Source/Programs/UnrealHeaderTool/Private/Classes.h b/Engine/Source/Programs/UnrealHeaderTool/Private/Classes.h
--- a/Engine/Source/Programs/UnrealHeaderTool/Private/Classes.h
+++ b/Engine/Source/Programs/UnrealHeaderTool/Private/Classes.h
@@ -1,9 +1,5 @@
 // Copyright 1998-2014 Epic Games, Inc. All Rights Reserved.
 
-/*=============================================================================
-	Classes.h: Container structure for all the Classes.
-=============================================================================*/
-
 #pragma once
 
 #include "ClassTree.h"
diff --git a/Engine/Source/Programs/UnrealHeaderTool/Private/HeaderParser.h b/Engine/Source/Programs/UnrealHeaderTool/Private/HeaderParser.h
--- a/Engine/Source/Programs/UnrealHeaderTool/Private/HeaderParser.h
+++ b/Engine/Source/Programs/UnrealHeaderTool/Private/HeaderParser.h
@@ -1,9 +1,5 @@
 // Copyright 1998-2014 Epic Games, Inc. All Rights Reserved.
 
-/*=============================================================================
-	HeaderParser.h: Parses annotated C++ headers to generate additional code and metadata.
-=============================================================================*/
-
 #pragma once
 
 #include "ParserHelper.h"
@@ -12,6 +8,7 @@
 class FClass;
 class FClasses;
 
+
 /*-----------------------------------------------------------------------------
 	Constants & types.
 -----------------------------------------------------------------------------*/
@@ -120,14 +117,13 @@ public:
 	}
 
 	/**
-	 * Transforms cpp-formated string containing default value, to inner formated string
-	 * If it cannot be transormed empty string is returned.
-	 *
-	 * @param Property - property, that owns the default value
-	 * @param CppForm -  cpp-formated string
-	 * @param out InnerForm -  inner formated string
+	 * Transforms CPP-formated string containing default value, to inner formated string
+	 * If it cannot be transformed empty string is returned.
 	 *
-	 * @return if succedded
+	 * @param Property The property that owns the default value.
+	 * @param CppForm A CPP-formated string.
+	 * @param out InnerForm Inner formated string
+	 * @return true on success, false otherwise.
 	 */
 	static bool DefaultValueStringCppFormatToInnerFormat(const UProperty* Property, const FString& CppForm, FString &InnerForm);
 
@@ -143,7 +139,7 @@ protected:
 	// Filename currently being parsed
 	FString Filename;
 
-	// Was the first include in the file a validly formed autogenerated header include?
+	// Was the first include in the file a validly formed auto-generated header include?
 	bool bSpottedAutogeneratedHeaderInclude;
 
 	// Current nest level, starts at 0.
@@ -163,7 +159,7 @@ protected:
 	{
 		enum Type 
 		{
-			// this directive is insignificant and does not change the codegeneration at all
+			// this directive is insignificant and does not change the code generation at all
 			Insignificant			= 0,
 			// this indicates we are in a WITH_EDITOR #if-Block
 			WithEditor				= 1<<0,
diff --git a/Engine/Source/Programs/UnrealHeaderTool/Private/Manifest.h b/Engine/Source/Programs/UnrealHeaderTool/Private/Manifest.h
--- a/Engine/Source/Programs/UnrealHeaderTool/Private/Manifest.h
+++ b/Engine/Source/Programs/UnrealHeaderTool/Private/Manifest.h
@@ -1,12 +1,8 @@
 // Copyright 1998-2014 Epic Games, Inc. All Rights Reserved.
 
-/*=============================================================================
-	Array.h: Data structure containing the contents of an
-	         UnrealHeaderTool.manifest file.
-=============================================================================*/
-
 #pragma once
 
+
 struct FManifestModule
 {
 	/** The name of the module */
@@ -46,6 +42,7 @@ struct FManifestModule
 	bool SaveExportedHeaders;
 };
 
+
 struct FManifest
 {
 	bool    IsGameTarget;
@@ -61,7 +58,6 @@ struct FManifest
 	 * Loads an UnrealHeaderTool.manifest from the specified filename.
 	 *
 	 * @param Filename The filename of the manifest to load.
-	 *
 	 * @return The loaded module info.
 	 */
 	static FManifest LoadFromFile(const FString& Filename);
diff --git a/Engine/Source/Programs/UnrealHeaderTool/Private/NativeClassExporter.h b/Engine/Source/Programs/UnrealHeaderTool/Private/NativeClassExporter.h
--- a/Engine/Source/Programs/UnrealHeaderTool/Private/NativeClassExporter.h
+++ b/Engine/Source/Programs/UnrealHeaderTool/Private/NativeClassExporter.h
@@ -1,9 +1,5 @@
 // Copyright 1998-2014 Epic Games, Inc. All Rights Reserved.
 
-/*=============================================================================
-	NativeClassExporter.h: Native class *Classes.h exporter header
-=============================================================================*/
-
 #pragma once
 
 #include "Class.h"
diff --git a/Engine/Source/Programs/UnrealHeaderTool/Private/StringUtils.h b/Engine/Source/Programs/UnrealHeaderTool/Private/StringUtils.h
--- a/Engine/Source/Programs/UnrealHeaderTool/Private/StringUtils.h
+++ b/Engine/Source/Programs/UnrealHeaderTool/Private/StringUtils.h
@@ -1,36 +1,33 @@
 // Copyright 1998-2014 Epic Games, Inc. All Rights Reserved.
 
-/*=============================================================================
-	StringUtils.h: Helper functions for manipulating strings.
-=============================================================================*/
-
 #pragma once
 
+
 /** 
  * Attempts to strip the given class name of its affixed prefix. If no prefix exists, will return a blank string.
  *
- * @param InClassName - Class Name with a prefix to be stripped
+ * @param InClassName Class Name with a prefix to be stripped
  */
 FString GetClassNameWithPrefixRemoved(const FString InClassName);
 
 /** 
  * Attempts to strip the given class name of its affixed prefix. If no prefix exists, will return unchanged string.
  *
- * @param InClassName - Class Name with a prefix to be stripped
+ * @param InClassName Class Name with a prefix to be stripped
  */
 FString GetClassNameWithoutPrefix(const FString& InClassNameOrFilename);
 
 /** 
  * Attempts to get class prefix. If the given class name does not start with a valid Unreal Prefix, it will return an empty string.
  *
- * @param InClassName - Name w/ potential prefix to check
+ * @param InClassName Name w/ potential prefix to check
  */
 FString GetClassPrefix(const FString InClassName);
 
 /** 
  * Attempts to get class prefix. If the given class name does not start with a valid Unreal Prefix, it will return an empty string.
  *
- * @param InClassName - Name w/ potential prefix to check
- * @param out bIsLabeledDeprecated - Reference param set to True if the class name is marked as deprecated
+ * @param InClassName Name w/ potential prefix to check
+ * @param out bIsLabeledDeprecated Reference param set to True if the class name is marked as deprecated
  */
 FString GetClassPrefix(const FString InClassName, bool& bIsLabeledDeprecated);
diff --git a/Engine/Source/Programs/UnrealLightmass/Private/CPUSolver/CPUSolver.cpp b/Engine/Source/Programs/UnrealLightmass/Private/CPUSolver/CPUSolver.cpp
--- a/Engine/Source/Programs/UnrealLightmass/Private/CPUSolver/CPUSolver.cpp
+++ b/Engine/Source/Programs/UnrealLightmass/Private/CPUSolver/CPUSolver.cpp
@@ -1,9 +1,5 @@
 // Copyright 1998-2014 Epic Games, Inc. All Rights Reserved.
 
-/*=============================================================================
-	CPUSolver.cpp: CPU solver implementation.
-=============================================================================*/
-
 #include "stdafx.h"
 #include "LightmassScene.h"
 #include "Importer.h"
diff --git a/Engine/Source/Programs/UnrealLightmass/Private/CPUSolver/CPUSolver.h b/Engine/Source/Programs/UnrealLightmass/Private/CPUSolver/CPUSolver.h
--- a/Engine/Source/Programs/UnrealLightmass/Private/CPUSolver/CPUSolver.h
+++ b/Engine/Source/Programs/UnrealLightmass/Private/CPUSolver/CPUSolver.h
@@ -1,10 +1,5 @@
 // Copyright 1998-2014 Epic Games, Inc. All Rights Reserved.
 
-/*=============================================================================
-	CPUSolver.h: CPU solver interface
-=============================================================================*/
-
-
 #pragma once
 
 #ifndef CPUSOLVER_H
diff --git a/Engine/Source/Programs/UnrealLightmass/Private/ImportExport/3DVisualizer.cpp b/Engine/Source/Programs/UnrealLightmass/Private/ImportExport/3DVisualizer.cpp
--- a/Engine/Source/Programs/UnrealLightmass/Private/ImportExport/3DVisualizer.cpp
+++ b/Engine/Source/Programs/UnrealLightmass/Private/ImportExport/3DVisualizer.cpp
@@ -1,9 +1,5 @@
 // Copyright 1998-2014 Epic Games, Inc. All Rights Reserved.
 
-/*=============================================================================
-	3DVisualizer.cpp: Helper class to visualize simple 3D geometry
-=============================================================================*/
-
 #include "stdafx.h"
 
 #if PLATFORM_WINDOWS
diff --git a/Engine/Source/Programs/UnrealLightmass/Private/ImportExport/3DVisualizer.h b/Engine/Source/Programs/UnrealLightmass/Private/ImportExport/3DVisualizer.h
--- a/Engine/Source/Programs/UnrealLightmass/Private/ImportExport/3DVisualizer.h
+++ b/Engine/Source/Programs/UnrealLightmass/Private/ImportExport/3DVisualizer.h
@@ -1,11 +1,8 @@
 // Copyright 1998-2014 Epic Games, Inc. All Rights Reserved.
 
-/*=============================================================================
-	3DVisualizer.h: Helper class to visualize simple 3D geometry
-=============================================================================*/
-
 #pragma once
 
+
 namespace Lightmass
 {
 
diff --git a/Engine/Source/Programs/UnrealLightmass/Private/ImportExport/Exporter.cpp b/Engine/Source/Programs/UnrealLightmass/Private/ImportExport/Exporter.cpp
--- a/Engine/Source/Programs/UnrealLightmass/Private/ImportExport/Exporter.cpp
+++ b/Engine/Source/Programs/UnrealLightmass/Private/ImportExport/Exporter.cpp
@@ -1,8 +1,5 @@
 // Copyright 1998-2014 Epic Games, Inc. All Rights Reserved.
 
-/*=============================================================================
-	Exporter.cpp: Lightmass solver exporter class.
-=============================================================================*/
 #include "stdafx.h"
 #include "LightmassSwarm.h"
 #include "Exporter.h"
diff --git a/Engine/Source/Programs/UnrealLightmass/Private/ImportExport/Exporter.h b/Engine/Source/Programs/UnrealLightmass/Private/ImportExport/Exporter.h
--- a/Engine/Source/Programs/UnrealLightmass/Private/ImportExport/Exporter.h
+++ b/Engine/Source/Programs/UnrealLightmass/Private/ImportExport/Exporter.h
@@ -1,13 +1,10 @@
 // Copyright 1998-2014 Epic Games, Inc. All Rights Reserved.
 
-/*=============================================================================
-	Exporter.h: Lightmass solver exporter class.
-=============================================================================*/
-
 #pragma once
 
 #include "LMCore.h"
 
+
 namespace Lightmass
 {
 	class FStaticLightingTextureMappingResult;
diff --git a/Engine/Source/Programs/UnrealLightmass/Private/ImportExport/Importer.cpp b/Engine/Source/Programs/UnrealLightmass/Private/ImportExport/Importer.cpp
--- a/Engine/Source/Programs/UnrealLightmass/Private/ImportExport/Importer.cpp
+++ b/Engine/Source/Programs/UnrealLightmass/Private/ImportExport/Importer.cpp
@@ -1,9 +1,5 @@
 // Copyright 1998-2014 Epic Games, Inc. All Rights Reserved.
 
-/*=============================================================================
-	Importer.cpp: Lightmass importer implementation.
-=============================================================================*/
-
 #include "stdafx.h"
 #include "LightmassSwarm.h"
 #include "Importer.h"
diff --git a/Engine/Source/Programs/UnrealLightmass/Private/ImportExport/Importer.h b/Engine/Source/Programs/UnrealLightmass/Private/ImportExport/Importer.h
--- a/Engine/Source/Programs/UnrealLightmass/Private/ImportExport/Importer.h
+++ b/Engine/Source/Programs/UnrealLightmass/Private/ImportExport/Importer.h
@@ -1,9 +1,5 @@
 // Copyright 1998-2014 Epic Games, Inc. All Rights Reserved.
 
-/*=============================================================================
-	Importer.h: Lightmass importer class.
-=============================================================================*/
-
 #pragma once
 
 #include "LMCore.h"
diff --git a/Engine/Source/Programs/UnrealLightmass/Private/ImportExport/LightmassScene.cpp b/Engine/Source/Programs/UnrealLightmass/Private/ImportExport/LightmassScene.cpp
--- a/Engine/Source/Programs/UnrealLightmass/Private/ImportExport/LightmassScene.cpp
+++ b/Engine/Source/Programs/UnrealLightmass/Private/ImportExport/LightmassScene.cpp
@@ -1,9 +1,5 @@
 // Copyright 1998-2014 Epic Games, Inc. All Rights Reserved.
 
-/*=============================================================================
-	Scene.cpp: Scene classes.
-=============================================================================*/
-
 #include "stdafx.h"
 #include "LightmassScene.h"
 #include "Importer.h"
diff --git a/Engine/Source/Programs/UnrealLightmass/Private/ImportExport/LightmassScene.h b/Engine/Source/Programs/UnrealLightmass/Private/ImportExport/LightmassScene.h
--- a/Engine/Source/Programs/UnrealLightmass/Private/ImportExport/LightmassScene.h
+++ b/Engine/Source/Programs/UnrealLightmass/Private/ImportExport/LightmassScene.h
@@ -1,12 +1,9 @@
 // Copyright 1998-2014 Epic Games, Inc. All Rights Reserved.
 
-/*=============================================================================
-	Scene.h: Scene classes.
-=============================================================================*/
-
 #pragma once
 
-#include "../Public/SceneExport.h"
+#include "SceneExport.h"
+
 
 namespace Lightmass
 {
diff --git a/Engine/Source/Programs/UnrealLightmass/Private/ImportExport/LightmassSwarm.cpp b/Engine/Source/Programs/UnrealLightmass/Private/ImportExport/LightmassSwarm.cpp
--- a/Engine/Source/Programs/UnrealLightmass/Private/ImportExport/LightmassSwarm.cpp
+++ b/Engine/Source/Programs/UnrealLightmass/Private/ImportExport/LightmassSwarm.cpp
@@ -1,13 +1,10 @@
 // Copyright 1998-2014 Epic Games, Inc. All Rights Reserved.
 
-/*=============================================================================
-	LightmassSwarm.cpp: Lightmass interface to swarm implementation.
-=============================================================================*/
-
 #include "stdafx.h"
 #include "LightmassSwarm.h"
 #include "CPUSolver.h"
 
+
 namespace Lightmass
 {
 
diff --git a/Engine/Source/Programs/UnrealLightmass/Private/ImportExport/LightmassSwarm.h b/Engine/Source/Programs/UnrealLightmass/Private/ImportExport/LightmassSwarm.h
--- a/Engine/Source/Programs/UnrealLightmass/Private/ImportExport/LightmassSwarm.h
+++ b/Engine/Source/Programs/UnrealLightmass/Private/ImportExport/LightmassSwarm.h
@@ -1,9 +1,5 @@
 // Copyright 1998-2014 Epic Games, Inc. All Rights Reserved.
 
-/*=============================================================================
-	LightmassSwarm.h: Lightmass interface to swarm.
-=============================================================================*/
-
 #pragma once
 
 #ifndef SWARMINTERFACE_API
diff --git a/Engine/Source/Programs/UnrealLightmass/Private/ImportExport/Material.cpp b/Engine/Source/Programs/UnrealLightmass/Private/ImportExport/Material.cpp
--- a/Engine/Source/Programs/UnrealLightmass/Private/ImportExport/Material.cpp
+++ b/Engine/Source/Programs/UnrealLightmass/Private/ImportExport/Material.cpp
@@ -1,8 +1,5 @@
 // Copyright 1998-2014 Epic Games, Inc. All Rights Reserved.
 
-/*=============================================================================
-	Material.cpp: Material classes implementation.
-=============================================================================*/
 #include "stdafx.h"
 #include "Material.h"
 #include "Importer.h"
diff --git a/Engine/Source/Programs/UnrealLightmass/Private/ImportExport/Material.h b/Engine/Source/Programs/UnrealLightmass/Private/ImportExport/Material.h
--- a/Engine/Source/Programs/UnrealLightmass/Private/ImportExport/Material.h
+++ b/Engine/Source/Programs/UnrealLightmass/Private/ImportExport/Material.h
@@ -1,13 +1,11 @@
 // Copyright 1998-2014 Epic Games, Inc. All Rights Reserved.
 
-/*=============================================================================
-	Material.h: Material classes.
-=============================================================================*/
 #pragma once
 
 // @todo UE4: This seems wrong to me to be needed (scene uses mesh, not the other way around), and is only needed for the MAX_TEXCOORDS
-#include "../Public/MaterialExport.h"
-#include "../Public/SceneExport.h"
+#include "MaterialExport.h"
+#include "SceneExport.h"
+
 
 namespace Lightmass
 {
diff --git a/Engine/Source/Programs/UnrealLightmass/Private/ImportExport/Mesh.cpp b/Engine/Source/Programs/UnrealLightmass/Private/ImportExport/Mesh.cpp
--- a/Engine/Source/Programs/UnrealLightmass/Private/ImportExport/Mesh.cpp
+++ b/Engine/Source/Programs/UnrealLightmass/Private/ImportExport/Mesh.cpp
@@ -1,9 +1,5 @@
 // Copyright 1998-2014 Epic Games, Inc. All Rights Reserved.
 
-/*=============================================================================
-	Mesh.cpp: Mesh classes.
-=============================================================================*/
-
 #include "stdafx.h"
 #include "Mesh.h"
 #include "Importer.h"
diff --git a/Engine/Source/Programs/UnrealLightmass/Private/ImportExport/Mesh.h b/Engine/Source/Programs/UnrealLightmass/Private/ImportExport/Mesh.h
--- a/Engine/Source/Programs/UnrealLightmass/Private/ImportExport/Mesh.h
+++ b/Engine/Source/Programs/UnrealLightmass/Private/ImportExport/Mesh.h
@@ -1,14 +1,11 @@
 // Copyright 1998-2014 Epic Games, Inc. All Rights Reserved.
 
-/*=============================================================================
-	Mesh.h: Mesh classes.
-=============================================================================*/
-
 #pragma once
 
 // @todo UE4: This seems wrong to me to be needed (scene uses mesh, not the other way around), and is only needed for the MAX_TEXCOORDS
-#include "../Public/MeshExport.h"
-#include "../Public/SceneExport.h"
+#include "MeshExport.h"
+#include "SceneExport.h"
+
 
 namespace Lightmass
 {
diff --git a/Engine/Source/Programs/UnrealLightmass/Private/Launch/UnitTest.cpp b/Engine/Source/Programs/UnrealLightmass/Private/Launch/UnitTest.cpp
--- a/Engine/Source/Programs/UnrealLightmass/Private/Launch/UnitTest.cpp
+++ b/Engine/Source/Programs/UnrealLightmass/Private/Launch/UnitTest.cpp
@@ -1,9 +1,5 @@
 // Copyright 1998-2014 Epic Games, Inc. All Rights Reserved.
 
-/*=============================================================================
-	UnitTest.cpp: Unit test implementation
-=============================================================================*/
-
 #include "stdafx.h"
 #include "UnitTest.h"
 
diff --git a/Engine/Source/Programs/UnrealLightmass/Private/Launch/UnitTest.h b/Engine/Source/Programs/UnrealLightmass/Private/Launch/UnitTest.h
--- a/Engine/Source/Programs/UnrealLightmass/Private/Launch/UnitTest.h
+++ b/Engine/Source/Programs/UnrealLightmass/Private/Launch/UnitTest.h
@@ -1,9 +1,5 @@
 // Copyright 1998-2014 Epic Games, Inc. All Rights Reserved.
 
-/*=============================================================================
-	UnitTest.h: Unit test interface
-=============================================================================*/
-
 #pragma once
 
 namespace Lightmass
diff --git a/Engine/Source/Programs/UnrealLightmass/Private/Launch/stdafx.h b/Engine/Source/Programs/UnrealLightmass/Private/Launch/stdafx.h
--- a/Engine/Source/Programs/UnrealLightmass/Private/Launch/stdafx.h
+++ b/Engine/Source/Programs/UnrealLightmass/Private/Launch/stdafx.h
@@ -1,9 +1,5 @@
 // Copyright 1998-2014 Epic Games, Inc. All Rights Reserved.
 
-/*=============================================================================
-	Scene.cpp: PCH master include file.
-=============================================================================*/
-
 #pragma once
 
 // basic stuff needed by everybody
diff --git a/Engine/Source/Programs/UnrealLightmass/Private/Lighting/BSP.cpp b/Engine/Source/Programs/UnrealLightmass/Private/Lighting/BSP.cpp
--- a/Engine/Source/Programs/UnrealLightmass/Private/Lighting/BSP.cpp
+++ b/Engine/Source/Programs/UnrealLightmass/Private/Lighting/BSP.cpp
@@ -1,9 +1,5 @@
 // Copyright 1998-2014 Epic Games, Inc. All Rights Reserved.
 
-/*=============================================================================
-	BSP.cpp: Static lighting BSP mesh/mapping implementation.
-=============================================================================*/
-
 #include "stdafx.h"
 #include "Importer.h"
 
diff --git a/Engine/Source/Programs/UnrealLightmass/Private/Lighting/BSP.h b/Engine/Source/Programs/UnrealLightmass/Private/Lighting/BSP.h
--- a/Engine/Source/Programs/UnrealLightmass/Private/Lighting/BSP.h
+++ b/Engine/Source/Programs/UnrealLightmass/Private/Lighting/BSP.h
@@ -1,9 +1,5 @@
 // Copyright 1998-2014 Epic Games, Inc. All Rights Reserved.
 
-/*=============================================================================
-	BSP.h: Static lighting BSP mesh/mapping definitions.
-=============================================================================*/
-
 #pragma once
 
 namespace Lightmass
diff --git a/Engine/Source/Programs/UnrealLightmass/Private/Lighting/BuildOptions.h b/Engine/Source/Programs/UnrealLightmass/Private/Lighting/BuildOptions.h
--- a/Engine/Source/Programs/UnrealLightmass/Private/Lighting/BuildOptions.h
+++ b/Engine/Source/Programs/UnrealLightmass/Private/Lighting/BuildOptions.h
@@ -1,9 +1,5 @@
 // Copyright 1998-2014 Epic Games, Inc. All Rights Reserved.
 
-/*=============================================================================
-	BuildOptions.h: Lighting options definitions.
-=============================================================================*/
-
 #pragma once
 
 namespace Lightmass
diff --git a/Engine/Source/Programs/UnrealLightmass/Private/Lighting/Cache.cpp b/Engine/Source/Programs/UnrealLightmass/Private/Lighting/Cache.cpp
--- a/Engine/Source/Programs/UnrealLightmass/Private/Lighting/Cache.cpp
+++ b/Engine/Source/Programs/UnrealLightmass/Private/Lighting/Cache.cpp
@@ -1,9 +1,5 @@
 // Copyright 1998-2014 Epic Games, Inc. All Rights Reserved.
 
-/*=============================================================================
-	Cache.cpp: Lighting cache implementation.
-=============================================================================*/
-
 #include "stdafx.h"
 #include "LightingSystem.h"
 
diff --git a/Engine/Source/Programs/UnrealLightmass/Private/Lighting/Cache.h b/Engine/Source/Programs/UnrealLightmass/Private/Lighting/Cache.h
--- a/Engine/Source/Programs/UnrealLightmass/Private/Lighting/Cache.h
+++ b/Engine/Source/Programs/UnrealLightmass/Private/Lighting/Cache.h
@@ -1,9 +1,5 @@
 // Copyright 1998-2014 Epic Games, Inc. All Rights Reserved.
 
-/*=============================================================================
-	Cache.h: Lighting cache definitions.
-=============================================================================*/
-
 #pragma once
 
 namespace Lightmass
diff --git a/Engine/Source/Programs/UnrealLightmass/Private/Lighting/Collision.cpp b/Engine/Source/Programs/UnrealLightmass/Private/Lighting/Collision.cpp
--- a/Engine/Source/Programs/UnrealLightmass/Private/Lighting/Collision.cpp
+++ b/Engine/Source/Programs/UnrealLightmass/Private/Lighting/Collision.cpp
@@ -1,9 +1,5 @@
 // Copyright 1998-2014 Epic Games, Inc. All Rights Reserved.
 
-/*=============================================================================
-	Collision.cpp: Static lighting collision system implementation.
-=============================================================================*/
-
 #include "stdafx.h"
 #include "LightingSystem.h"
 
diff --git a/Engine/Source/Programs/UnrealLightmass/Private/Lighting/Collision.h b/Engine/Source/Programs/UnrealLightmass/Private/Lighting/Collision.h
--- a/Engine/Source/Programs/UnrealLightmass/Private/Lighting/Collision.h
+++ b/Engine/Source/Programs/UnrealLightmass/Private/Lighting/Collision.h
@@ -1,9 +1,5 @@
 // Copyright 1998-2014 Epic Games, Inc. All Rights Reserved.
 
-/*=============================================================================
-	Collision.h: Static lighting collision system definitions.
-=============================================================================*/
-
 #pragma once
 
 #include "LMOctree.h"
diff --git a/Engine/Source/Programs/UnrealLightmass/Private/Lighting/FinalGather.cpp b/Engine/Source/Programs/UnrealLightmass/Private/Lighting/FinalGather.cpp
--- a/Engine/Source/Programs/UnrealLightmass/Private/Lighting/FinalGather.cpp
+++ b/Engine/Source/Programs/UnrealLightmass/Private/Lighting/FinalGather.cpp
@@ -1,9 +1,5 @@
 // Copyright 1998-2014 Epic Games, Inc. All Rights Reserved.
 
-/*=============================================================================
-	FinalGather.cpp: Final gathering functionality
-=============================================================================*/
-
 #include "stdafx.h"
 #include "Exporter.h"
 #include "LightmassSwarm.h"
diff --git a/Engine/Source/Programs/UnrealLightmass/Private/Lighting/FluidSurface.cpp b/Engine/Source/Programs/UnrealLightmass/Private/Lighting/FluidSurface.cpp
--- a/Engine/Source/Programs/UnrealLightmass/Private/Lighting/FluidSurface.cpp
+++ b/Engine/Source/Programs/UnrealLightmass/Private/Lighting/FluidSurface.cpp
@@ -1,8 +1,5 @@
 // Copyright 1998-2014 Epic Games, Inc. All Rights Reserved.
 
-/*=============================================================================
-	FluidSurface.cpp: Static lighting FluidSurface mesh/mapping implementation.
-=============================================================================*/
 #include "stdafx.h"
 #include "Importer.h"
 #include "FluidSurface.h"
diff --git a/Engine/Source/Programs/UnrealLightmass/Private/Lighting/FluidSurface.h b/Engine/Source/Programs/UnrealLightmass/Private/Lighting/FluidSurface.h
--- a/Engine/Source/Programs/UnrealLightmass/Private/Lighting/FluidSurface.h
+++ b/Engine/Source/Programs/UnrealLightmass/Private/Lighting/FluidSurface.h
@@ -1,9 +1,5 @@
 // Copyright 1998-2014 Epic Games, Inc. All Rights Reserved.
 
-/*=============================================================================
-	FluidSurface.h: Static lighting FluidSurface mesh/mapping definitions.
-=============================================================================*/
-
 #pragma once
 
 namespace Lightmass
diff --git a/Engine/Source/Programs/UnrealLightmass/Private/Lighting/Landscape.cpp b/Engine/Source/Programs/UnrealLightmass/Private/Lighting/Landscape.cpp
--- a/Engine/Source/Programs/UnrealLightmass/Private/Lighting/Landscape.cpp
+++ b/Engine/Source/Programs/UnrealLightmass/Private/Lighting/Landscape.cpp
@@ -1,8 +1,5 @@
 // Copyright 1998-2014 Epic Games, Inc. All Rights Reserved.
 
-/*=============================================================================
-	Landscape.cpp: Static lighting Landscape mesh/mapping implementation.
-=============================================================================*/
 #include "stdafx.h"
 #include "Importer.h"
 #include "Landscape.h"
diff --git a/Engine/Source/Programs/UnrealLightmass/Private/Lighting/Landscape.h b/Engine/Source/Programs/UnrealLightmass/Private/Lighting/Landscape.h
--- a/Engine/Source/Programs/UnrealLightmass/Private/Lighting/Landscape.h
+++ b/Engine/Source/Programs/UnrealLightmass/Private/Lighting/Landscape.h
@@ -1,9 +1,5 @@
 // Copyright 1998-2014 Epic Games, Inc. All Rights Reserved.
 
-/*=============================================================================
-	Landscape.h: Static lighting Landscape mesh/mapping definitions.
-=============================================================================*/
-
 #pragma once
 
 #define LANDSCAPE_ZSCALE	(1.0f/128.0f)
diff --git a/Engine/Source/Programs/UnrealLightmass/Private/Lighting/Lighting.h b/Engine/Source/Programs/UnrealLightmass/Private/Lighting/Lighting.h
--- a/Engine/Source/Programs/UnrealLightmass/Private/Lighting/Lighting.h
+++ b/Engine/Source/Programs/UnrealLightmass/Private/Lighting/Lighting.h
@@ -1,9 +1,5 @@
 // Copyright 1998-2014 Epic Games, Inc. All Rights Reserved.
 
-/*=============================================================================
-	Lighting.h: Private static lighting system includes.
-=============================================================================*/
-
 #pragma once
 
 #include "LMCore.h"
diff --git a/Engine/Source/Programs/UnrealLightmass/Private/Lighting/LightingMesh.cpp b/Engine/Source/Programs/UnrealLightmass/Private/Lighting/LightingMesh.cpp
--- a/Engine/Source/Programs/UnrealLightmass/Private/Lighting/LightingMesh.cpp
+++ b/Engine/Source/Programs/UnrealLightmass/Private/Lighting/LightingMesh.cpp
@@ -1,9 +1,5 @@
 // Copyright 1998-2014 Epic Games, Inc. All Rights Reserved.
 
-/*=============================================================================
-	LightingMesh.cpp: Static lighting mesh implementation.
-=============================================================================*/
-
 #include "stdafx.h"
 #include "Importer.h"
 #include "MonteCarlo.h"
diff --git a/Engine/Source/Programs/UnrealLightmass/Private/Lighting/LightingMesh.h b/Engine/Source/Programs/UnrealLightmass/Private/Lighting/LightingMesh.h
--- a/Engine/Source/Programs/UnrealLightmass/Private/Lighting/LightingMesh.h
+++ b/Engine/Source/Programs/UnrealLightmass/Private/Lighting/LightingMesh.h
@@ -1,9 +1,5 @@
 // Copyright 1998-2014 Epic Games, Inc. All Rights Reserved.
 
-/*=============================================================================
-	LightingMesh.h: Static lighting mesh definitions.
-=============================================================================*/
-
 #pragma once
 
 #include "ImportExport.h"
diff --git a/Engine/Source/Programs/UnrealLightmass/Private/Lighting/LightingSystem.cpp b/Engine/Source/Programs/UnrealLightmass/Private/Lighting/LightingSystem.cpp
--- a/Engine/Source/Programs/UnrealLightmass/Private/Lighting/LightingSystem.cpp
+++ b/Engine/Source/Programs/UnrealLightmass/Private/Lighting/LightingSystem.cpp
@@ -1,9 +1,5 @@
 // Copyright 1998-2014 Epic Games, Inc. All Rights Reserved.
 
-/*=============================================================================
-	LightingSystem.cpp: Private static lighting system implementation.
-=============================================================================*/
-
 #include "stdafx.h"
 #include "Exporter.h"
 #include "LightmassSwarm.h"
diff --git a/Engine/Source/Programs/UnrealLightmass/Private/Lighting/LightingSystem.h b/Engine/Source/Programs/UnrealLightmass/Private/Lighting/LightingSystem.h
--- a/Engine/Source/Programs/UnrealLightmass/Private/Lighting/LightingSystem.h
+++ b/Engine/Source/Programs/UnrealLightmass/Private/Lighting/LightingSystem.h
@@ -1,9 +1,5 @@
 // Copyright 1998-2014 Epic Games, Inc. All Rights Reserved.
 
-/*=============================================================================
-	LightingSystem.h: Private static lighting system definitions.
-=============================================================================*/
-
 #pragma once
 
 #include "CPUSolver.h"
diff --git a/Engine/Source/Programs/UnrealLightmass/Private/Lighting/LightmapData.cpp b/Engine/Source/Programs/UnrealLightmass/Private/Lighting/LightmapData.cpp
--- a/Engine/Source/Programs/UnrealLightmass/Private/Lighting/LightmapData.cpp
+++ b/Engine/Source/Programs/UnrealLightmass/Private/Lighting/LightmapData.cpp
@@ -1,9 +1,5 @@
 // Copyright 1998-2014 Epic Games, Inc. All Rights Reserved.
 
-/*=============================================================================
-	LightmapData.cpp: Lightmap data implementation.
-=============================================================================*/
-
 #include "stdafx.h"
 #include "Exporter.h"
 #include "LightmassSwarm.h"
diff --git a/Engine/Source/Programs/UnrealLightmass/Private/Lighting/LightmapData.h b/Engine/Source/Programs/UnrealLightmass/Private/Lighting/LightmapData.h
--- a/Engine/Source/Programs/UnrealLightmass/Private/Lighting/LightmapData.h
+++ b/Engine/Source/Programs/UnrealLightmass/Private/Lighting/LightmapData.h
@@ -1,12 +1,9 @@
 // Copyright 1998-2014 Epic Games, Inc. All Rights Reserved.
 
-/*=============================================================================
-	LightmapData.h: Lightmap data definitions.
-=============================================================================*/
-
 #pragma once
 
-#include "../Public/ImportExport.h"
+#include "ImportExport.h"
+
 
 namespace Lightmass
 {
diff --git a/Engine/Source/Programs/UnrealLightmass/Private/Lighting/Mappings.cpp b/Engine/Source/Programs/UnrealLightmass/Private/Lighting/Mappings.cpp
--- a/Engine/Source/Programs/UnrealLightmass/Private/Lighting/Mappings.cpp
+++ b/Engine/Source/Programs/UnrealLightmass/Private/Lighting/Mappings.cpp
@@ -1,9 +1,5 @@
 // Copyright 1998-2014 Epic Games, Inc. All Rights Reserved.
 
-/*=============================================================================
-	Mappings.cpp: Static lighting mapping implementation.
-=============================================================================*/
-
 #include "stdafx.h"
 #include "Importer.h"
 
diff --git a/Engine/Source/Programs/UnrealLightmass/Private/Lighting/Mappings.h b/Engine/Source/Programs/UnrealLightmass/Private/Lighting/Mappings.h
--- a/Engine/Source/Programs/UnrealLightmass/Private/Lighting/Mappings.h
+++ b/Engine/Source/Programs/UnrealLightmass/Private/Lighting/Mappings.h
@@ -1,9 +1,5 @@
 // Copyright 1998-2014 Epic Games, Inc. All Rights Reserved.
 
-/*=============================================================================
-	Mappings.h: Static lighting mapping definitions.
-=============================================================================*/
-
 #pragma once
 
 #include "LockFreeList.h"
diff --git a/Engine/Source/Programs/UnrealLightmass/Private/Lighting/MonteCarlo.cpp b/Engine/Source/Programs/UnrealLightmass/Private/Lighting/MonteCarlo.cpp
--- a/Engine/Source/Programs/UnrealLightmass/Private/Lighting/MonteCarlo.cpp
+++ b/Engine/Source/Programs/UnrealLightmass/Private/Lighting/MonteCarlo.cpp
@@ -1,9 +1,5 @@
 // Copyright 1998-2014 Epic Games, Inc. All Rights Reserved.
 
-/*=============================================================================
-	MonteCarlo.cpp: Utilities for Monte Carlo Integration.
-=============================================================================*/
-
 #include "stdafx.h"
 #include "MonteCarlo.h"
 
diff --git a/Engine/Source/Programs/UnrealLightmass/Private/Lighting/MonteCarlo.h b/Engine/Source/Programs/UnrealLightmass/Private/Lighting/MonteCarlo.h
--- a/Engine/Source/Programs/UnrealLightmass/Private/Lighting/MonteCarlo.h
+++ b/Engine/Source/Programs/UnrealLightmass/Private/Lighting/MonteCarlo.h
@@ -1,11 +1,8 @@
 // Copyright 1998-2014 Epic Games, Inc. All Rights Reserved.
 
-/*=============================================================================
-	MonteCarlo.h: Utilities for Monte Carlo integration.
-=============================================================================*/
-
 #pragma once
 
+
 namespace Lightmass
 {
 	/** Generates valid X and Y axes of a coordinate system, given the Z axis. */
diff --git a/Engine/Source/Programs/UnrealLightmass/Private/Lighting/PhotonMapping.cpp b/Engine/Source/Programs/UnrealLightmass/Private/Lighting/PhotonMapping.cpp
--- a/Engine/Source/Programs/UnrealLightmass/Private/Lighting/PhotonMapping.cpp
+++ b/Engine/Source/Programs/UnrealLightmass/Private/Lighting/PhotonMapping.cpp
@@ -1,9 +1,5 @@
 // Copyright 1998-2014 Epic Games, Inc. All Rights Reserved.
 
-/*=============================================================================
-	Photon.cpp: Static lighting photon mapping implementation.
-=============================================================================*/
-
 #include "stdafx.h"
 #include "LightingSystem.h"
 #include "MonteCarlo.h"
diff --git a/Engine/Source/Programs/UnrealLightmass/Private/Lighting/Raster.h b/Engine/Source/Programs/UnrealLightmass/Private/Lighting/Raster.h
--- a/Engine/Source/Programs/UnrealLightmass/Private/Lighting/Raster.h
+++ b/Engine/Source/Programs/UnrealLightmass/Private/Lighting/Raster.h
@@ -1,10 +1,5 @@
 // Copyright 1998-2014 Epic Games, Inc. All Rights Reserved.
 
-/*=============================================================================
-	Raster.h: Triangle rasterization code
-=============================================================================*/
-
-
 #pragma once
 
 namespace Lightmass
diff --git a/Engine/Source/Programs/UnrealLightmass/Private/Lighting/StaticMesh.cpp b/Engine/Source/Programs/UnrealLightmass/Private/Lighting/StaticMesh.cpp
--- a/Engine/Source/Programs/UnrealLightmass/Private/Lighting/StaticMesh.cpp
+++ b/Engine/Source/Programs/UnrealLightmass/Private/Lighting/StaticMesh.cpp
@@ -1,9 +1,5 @@
 // Copyright 1998-2014 Epic Games, Inc. All Rights Reserved.
 
-/*=============================================================================
-	StaticMesh.cpp: Static lighting StaticMesh mesh/mapping implementation.
-=============================================================================*/
-
 #include "stdafx.h"
 #include "Importer.h"
 
diff --git a/Engine/Source/Programs/UnrealLightmass/Private/Lighting/StaticMesh.h b/Engine/Source/Programs/UnrealLightmass/Private/Lighting/StaticMesh.h
--- a/Engine/Source/Programs/UnrealLightmass/Private/Lighting/StaticMesh.h
+++ b/Engine/Source/Programs/UnrealLightmass/Private/Lighting/StaticMesh.h
@@ -1,9 +1,5 @@
 // Copyright 1998-2014 Epic Games, Inc. All Rights Reserved.
 
-/*=============================================================================
-	StaticMesh.h: Static lighting StaticMesh mesh/mapping definitions.
-=============================================================================*/
-
 #pragma once
 
 namespace Lightmass
diff --git a/Engine/Source/Programs/UnrealLightmass/Private/Lighting/Texture.h b/Engine/Source/Programs/UnrealLightmass/Private/Lighting/Texture.h
--- a/Engine/Source/Programs/UnrealLightmass/Private/Lighting/Texture.h
+++ b/Engine/Source/Programs/UnrealLightmass/Private/Lighting/Texture.h
@@ -1,9 +1,5 @@
 // Copyright 1998-2014 Epic Games, Inc. All Rights Reserved.
 
-/*=============================================================================
-	Texture.h: Texture definitions.
-=============================================================================*/
-
 #pragma once
 
 namespace Lightmass
diff --git a/Engine/Source/Programs/UnrealLightmass/Private/Lighting/TextureMapping.cpp b/Engine/Source/Programs/UnrealLightmass/Private/Lighting/TextureMapping.cpp
--- a/Engine/Source/Programs/UnrealLightmass/Private/Lighting/TextureMapping.cpp
+++ b/Engine/Source/Programs/UnrealLightmass/Private/Lighting/TextureMapping.cpp
@@ -1,9 +1,5 @@
 // Copyright 1998-2014 Epic Games, Inc. All Rights Reserved.
 
-/*=============================================================================
-	TextureMapping.cpp: Static lighting texture mapping implementation.
-=============================================================================*/
-
 #include "stdafx.h"
 #include "Raster.h"
 #include "LightingSystem.h"
diff --git a/Engine/Source/Programs/UnrealLightmass/Private/LightmassCore/LMCore.cpp b/Engine/Source/Programs/UnrealLightmass/Private/LightmassCore/LMCore.cpp
--- a/Engine/Source/Programs/UnrealLightmass/Private/LightmassCore/LMCore.cpp
+++ b/Engine/Source/Programs/UnrealLightmass/Private/LightmassCore/LMCore.cpp
@@ -1,9 +1,5 @@
 // Copyright 1998-2014 Epic Games, Inc. All Rights Reserved.
 
-/*=============================================================================
-	LMCore.h: Some implementation of LMCore functionality
-=============================================================================*/
-
 #include "stdafx.h"
 #include "LMCore.h"
 
diff --git a/Engine/Source/Programs/UnrealLightmass/Private/LightmassCore/LMCore.h b/Engine/Source/Programs/UnrealLightmass/Private/LightmassCore/LMCore.h
--- a/Engine/Source/Programs/UnrealLightmass/Private/LightmassCore/LMCore.h
+++ b/Engine/Source/Programs/UnrealLightmass/Private/LightmassCore/LMCore.h
@@ -1,9 +1,5 @@
 // Copyright 1998-2014 Epic Games, Inc. All Rights Reserved.
 
-/*=============================================================================
-	LMCore.h: Most basic include file for LMCore 
-=============================================================================*/
-
 #pragma once
 
 // recommended that these are included by all files (ie, leave in LMCore.h, and maybe include LMCore.h in stdafx.h)
diff --git a/Engine/Source/Programs/UnrealLightmass/Private/LightmassCore/Math/LMCollision.h b/Engine/Source/Programs/UnrealLightmass/Private/LightmassCore/Math/LMCollision.h
--- a/Engine/Source/Programs/UnrealLightmass/Private/LightmassCore/Math/LMCollision.h
+++ b/Engine/Source/Programs/UnrealLightmass/Private/LightmassCore/Math/LMCollision.h
@@ -1,9 +1,5 @@
 // Copyright 1998-2014 Epic Games, Inc. All Rights Reserved.
 
-/*=============================================================================
-	LMCollision.h: Collision definitions.
-=============================================================================*/
-
 #pragma once
 
 namespace Lightmass
diff --git a/Engine/Source/Programs/UnrealLightmass/Private/LightmassCore/Math/LMMath.cpp b/Engine/Source/Programs/UnrealLightmass/Private/LightmassCore/Math/LMMath.cpp
--- a/Engine/Source/Programs/UnrealLightmass/Private/LightmassCore/Math/LMMath.cpp
+++ b/Engine/Source/Programs/UnrealLightmass/Private/LightmassCore/Math/LMMath.cpp
@@ -1,9 +1,5 @@
 // Copyright 1998-2014 Epic Games, Inc. All Rights Reserved.
 
-/*=============================================================================
-	LMMath.h: Some implementation of LM math functionality
-=============================================================================*/
-
 #include "stdafx.h"
 #include "LMCore.h"
 
diff --git a/Engine/Source/Programs/UnrealLightmass/Private/LightmassCore/Math/LMMath.h b/Engine/Source/Programs/UnrealLightmass/Private/LightmassCore/Math/LMMath.h
--- a/Engine/Source/Programs/UnrealLightmass/Private/LightmassCore/Math/LMMath.h
+++ b/Engine/Source/Programs/UnrealLightmass/Private/LightmassCore/Math/LMMath.h
@@ -1,9 +1,5 @@
 // Copyright 1998-2014 Epic Games, Inc. All Rights Reserved.
 
-/*=============================================================================
-	LMMath.h: Main include file for math routines and classes
-=============================================================================*/
-
 #pragma once
 
 namespace Lightmass
diff --git a/Engine/Source/Programs/UnrealLightmass/Private/LightmassCore/Math/LMMathSSE.h b/Engine/Source/Programs/UnrealLightmass/Private/LightmassCore/Math/LMMathSSE.h
--- a/Engine/Source/Programs/UnrealLightmass/Private/LightmassCore/Math/LMMathSSE.h
+++ b/Engine/Source/Programs/UnrealLightmass/Private/LightmassCore/Math/LMMathSSE.h
@@ -1,10 +1,5 @@
 // Copyright 1998-2014 Epic Games, Inc. All Rights Reserved.
 
-/*=============================================================================
-	LMMathSSE.h: Intel SSE specific vector intrinsics (also using D3D9X)
-=============================================================================*/
-
-
 #pragma once
 
 
diff --git a/Engine/Source/Programs/UnrealLightmass/Private/LightmassCore/Math/LMOctree.h b/Engine/Source/Programs/UnrealLightmass/Private/LightmassCore/Math/LMOctree.h
--- a/Engine/Source/Programs/UnrealLightmass/Private/LightmassCore/Math/LMOctree.h
+++ b/Engine/Source/Programs/UnrealLightmass/Private/LightmassCore/Math/LMOctree.h
@@ -1,9 +1,5 @@
 // Copyright 1998-2014 Epic Games, Inc. All Rights Reserved.
 
-/*=============================================================================
-	LMOctree.h: Generic octree definition.
-=============================================================================*/
-
 #pragma once
 
 namespace Lightmass
diff --git a/Engine/Source/Programs/UnrealLightmass/Private/LightmassCore/Math/LMkDOP.h b/Engine/Source/Programs/UnrealLightmass/Private/LightmassCore/Math/LMkDOP.h
--- a/Engine/Source/Programs/UnrealLightmass/Private/LightmassCore/Math/LMkDOP.h
+++ b/Engine/Source/Programs/UnrealLightmass/Private/LightmassCore/Math/LMkDOP.h
@@ -1,9 +1,5 @@
 // Copyright 1998-2014 Epic Games, Inc. All Rights Reserved.
 
-/*=============================================================================
-	LMOctree.h: Generic octree definition.
-=============================================================================*/
-
 #pragma once
 
 // #include "Collision.h"
diff --git a/Engine/Source/Programs/UnrealLightmass/Private/LightmassCore/Misc/LMDebug.cpp b/Engine/Source/Programs/UnrealLightmass/Private/LightmassCore/Misc/LMDebug.cpp
--- a/Engine/Source/Programs/UnrealLightmass/Private/LightmassCore/Misc/LMDebug.cpp
+++ b/Engine/Source/Programs/UnrealLightmass/Private/LightmassCore/Misc/LMDebug.cpp
@@ -1,9 +1,5 @@
 // Copyright 1998-2014 Epic Games, Inc. All Rights Reserved.
 
-/*=============================================================================
-	LMMath.h: Some implementation of LM math functionality
-=============================================================================*/
-
 #include "stdafx.h"
 #include "LMCore.h"
 #include "LMDebug.h"
diff --git a/Engine/Source/Programs/UnrealLightmass/Private/LightmassCore/Misc/LMDebug.h b/Engine/Source/Programs/UnrealLightmass/Private/LightmassCore/Misc/LMDebug.h
--- a/Engine/Source/Programs/UnrealLightmass/Private/LightmassCore/Misc/LMDebug.h
+++ b/Engine/Source/Programs/UnrealLightmass/Private/LightmassCore/Misc/LMDebug.h
@@ -1,8 +1,5 @@
 // Copyright 1998-2014 Epic Games, Inc. All Rights Reserved.
 
-/*=============================================================================
-	LMMath.h: Some implementation of LM math functionality
-=============================================================================*/
 #pragma once
 
 namespace Lightmass
diff --git a/Engine/Source/Programs/UnrealLightmass/Private/LightmassCore/Misc/LMHelpers.h b/Engine/Source/Programs/UnrealLightmass/Private/LightmassCore/Misc/LMHelpers.h
--- a/Engine/Source/Programs/UnrealLightmass/Private/LightmassCore/Misc/LMHelpers.h
+++ b/Engine/Source/Programs/UnrealLightmass/Private/LightmassCore/Misc/LMHelpers.h
@@ -1,9 +1,5 @@
 // Copyright 1998-2014 Epic Games, Inc. All Rights Reserved.
 
-/*=============================================================================
-	LMHelpers.h: Mostly defines
-=============================================================================*/
-
 #pragma once
 
 namespace Lightmass
diff --git a/Engine/Source/Programs/UnrealLightmass/Private/LightmassCore/Misc/LMStats.h b/Engine/Source/Programs/UnrealLightmass/Private/LightmassCore/Misc/LMStats.h
--- a/Engine/Source/Programs/UnrealLightmass/Private/LightmassCore/Misc/LMStats.h
+++ b/Engine/Source/Programs/UnrealLightmass/Private/LightmassCore/Misc/LMStats.h
@@ -1,9 +1,5 @@
 // Copyright 1998-2014 Epic Games, Inc. All Rights Reserved.
 
-/*=============================================================================
-	LMStats.h: Stat classes
-=============================================================================*/
-
 #pragma once
 
 
diff --git a/Engine/Source/Programs/UnrealLightmass/Private/LightmassCore/Misc/LMThreading.h b/Engine/Source/Programs/UnrealLightmass/Private/LightmassCore/Misc/LMThreading.h
--- a/Engine/Source/Programs/UnrealLightmass/Private/LightmassCore/Misc/LMThreading.h
+++ b/Engine/Source/Programs/UnrealLightmass/Private/LightmassCore/Misc/LMThreading.h
@@ -1,9 +1,5 @@
 // Copyright 1998-2014 Epic Games, Inc. All Rights Reserved.
 
-/*=============================================================================
-	LMThreading.h: Threading/process functionality
-=============================================================================*/
-
 #pragma once
 
 #if PLATFORM_USE_PTHREADS
diff --git a/Engine/Source/Programs/UnrealLightmass/Private/LightmassCore/Templates/LMQueue.h b/Engine/Source/Programs/UnrealLightmass/Private/LightmassCore/Templates/LMQueue.h
--- a/Engine/Source/Programs/UnrealLightmass/Private/LightmassCore/Templates/LMQueue.h
+++ b/Engine/Source/Programs/UnrealLightmass/Private/LightmassCore/Templates/LMQueue.h
@@ -1,9 +1,5 @@
 // Copyright 1998-2014 Epic Games, Inc. All Rights Reserved.
 
-/*=============================================================================
-	LMQueue.h: Static FIFO queue (TQueue) definitions.
-=============================================================================*/
-
 #pragma once
 
 namespace Lightmass
diff --git a/Engine/Source/Programs/UnrealLightmass/Public/FluidExport.h b/Engine/Source/Programs/UnrealLightmass/Public/FluidExport.h
--- a/Engine/Source/Programs/UnrealLightmass/Public/FluidExport.h
+++ b/Engine/Source/Programs/UnrealLightmass/Public/FluidExport.h
@@ -1,9 +1,5 @@
 // Copyright 1998-2014 Epic Games, Inc. All Rights Reserved.
 
-/*=============================================================================
-	FluidExport.h: Fluid export data definitions.
-=============================================================================*/
-
 #pragma once
 
 namespace Lightmass
diff --git a/Engine/Source/Programs/UnrealLightmass/Public/ImportExport.h b/Engine/Source/Programs/UnrealLightmass/Public/ImportExport.h
--- a/Engine/Source/Programs/UnrealLightmass/Public/ImportExport.h
+++ b/Engine/Source/Programs/UnrealLightmass/Public/ImportExport.h
@@ -1,9 +1,5 @@
 // Copyright 1998-2014 Epic Games, Inc. All Rights Reserved.
 
-/*=============================================================================
-	ImportExport.h: Import proper export data definitions.
-=============================================================================*/
-
 #pragma once
 
 namespace Lightmass
diff --git a/Engine/Source/Programs/UnrealLightmass/Public/LightmassPublic.h b/Engine/Source/Programs/UnrealLightmass/Public/LightmassPublic.h
--- a/Engine/Source/Programs/UnrealLightmass/Public/LightmassPublic.h
+++ b/Engine/Source/Programs/UnrealLightmass/Public/LightmassPublic.h
@@ -1,9 +1,5 @@
 // Copyright 1998-2014 Epic Games, Inc. All Rights Reserved.
 
-/*=============================================================================
-	LightmassPublic.h: Main header file for external users of Lightmass to include
-=============================================================================*/
-
 #pragma once
 
 // include all of the headers needed by external users
diff --git a/Engine/Source/Programs/UnrealLightmass/Public/MaterialExport.h b/Engine/Source/Programs/UnrealLightmass/Public/MaterialExport.h
--- a/Engine/Source/Programs/UnrealLightmass/Public/MaterialExport.h
+++ b/Engine/Source/Programs/UnrealLightmass/Public/MaterialExport.h
@@ -1,9 +1,5 @@
 // Copyright 1998-2014 Epic Games, Inc. All Rights Reserved.
 
-/*=============================================================================
-	MaterialExport.h: Material export data definitions.
-=============================================================================*/
-
 #pragma once
 
 namespace Lightmass
diff --git a/Engine/Source/Programs/UnrealLightmass/Public/MeshExport.h b/Engine/Source/Programs/UnrealLightmass/Public/MeshExport.h
--- a/Engine/Source/Programs/UnrealLightmass/Public/MeshExport.h
+++ b/Engine/Source/Programs/UnrealLightmass/Public/MeshExport.h
@@ -1,9 +1,5 @@
 // Copyright 1998-2014 Epic Games, Inc. All Rights Reserved.
 
-/*=============================================================================
-	MeshExport.h: Mesh export data definitions.
-=============================================================================*/
-
 #pragma once
 
 namespace Lightmass
diff --git a/Engine/Source/Programs/UnrealLightmass/Public/SceneExport.h b/Engine/Source/Programs/UnrealLightmass/Public/SceneExport.h
--- a/Engine/Source/Programs/UnrealLightmass/Public/SceneExport.h
+++ b/Engine/Source/Programs/UnrealLightmass/Public/SceneExport.h
@@ -1,9 +1,5 @@
 // Copyright 1998-2014 Epic Games, Inc. All Rights Reserved.
 
-/*=============================================================================
-	SceneExport.h: Scene export data definitions.
-=============================================================================*/
-
 #pragma once
 
 #include "MeshExport.h"
diff --git a/Engine/Source/Runtime/AutomationMessages/Classes/AutomationWorkerMessages.h b/Engine/Source/Runtime/AutomationMessages/Classes/AutomationWorkerMessages.h
--- a/Engine/Source/Runtime/AutomationMessages/Classes/AutomationWorkerMessages.h
+++ b/Engine/Source/Runtime/AutomationMessages/Classes/AutomationWorkerMessages.h
@@ -1,9 +1,5 @@
 // Copyright 1998-2014 Epic Games, Inc. All Rights Reserved.
 
-/*=============================================================================
-	AutomationWorkerMessages.h: Declares message types used by automation workers.
-=============================================================================*/
-
 #pragma once
 
 #include "AutomationWorkerMessages.generated.h"
@@ -20,39 +16,26 @@ struct FAutomationWorkerFindWorkers
 {
 	GENERATED_USTRUCT_BODY()
 
-	/**
-	 * Holds the change list number to find workers for.
-	 */
+	/** Holds the change list number to find workers for. */
 	UPROPERTY()
 	int32 Changelist;
 
-	/**
-	 * The name of the game.
-	 */
+	/** The name of the game. */
 	UPROPERTY()
 	FString GameName;
 
-	/**
-	 * The name of the process.
-	 */
+	/** The name of the process. */
 	UPROPERTY()
 	FString ProcessName;
 
-	/**
-	 * Holds the session identifier to find workers for.
-	 */
+	/** Holds the session identifier to find workers for. */
 	UPROPERTY()
 	FGuid SessionId;
 
-
-	/**
-	 * Default constructor.
-	 */
+	/** Default constructor. */
 	FAutomationWorkerFindWorkers( ) { }
 
-	/**
-	 * Creates and initializes a new instance.
-	 */
+	/** Creates and initializes a new instance. */
 	FAutomationWorkerFindWorkers( int32 InChangelist, const FString& InGameName, const FString& InProcessName, const FGuid& InSessionId )
 		: Changelist(InChangelist)
 		, GameName(InGameName)
@@ -61,6 +44,7 @@ struct FAutomationWorkerFindWorkers
 	{ }
 };
 
+
 template<>
 struct TStructOpsTypeTraits<FAutomationWorkerFindWorkers> : public TStructOpsTypeTraitsBase
 {
@@ -79,67 +63,46 @@ struct FAutomationWorkerFindWorkersResponse
 {
 	GENERATED_USTRUCT_BODY()
 
-	/**
-	 * Holds the name of the device that the worker is running on.
-	 */
+	/** Holds the name of the device that the worker is running on. */
 	UPROPERTY()
 	FString DeviceName;
 
-	/**
-	 * Holds the name of the worker's application instance.
-	 */
+	/** Holds the name of the worker's application instance. */
 	UPROPERTY()
 	FString InstanceName;
 
-	/**
-	 * Holds the name of the platform that the worker is running on.
-	 */
+	/** Holds the name of the platform that the worker is running on. */
 	UPROPERTY()
 	FString Platform;
 
-	/**
-	 * Holds the name of the operating system version.
-	 */
+	/** Holds the name of the operating system version. */
 	UPROPERTY()
 	FString OSVersionName;
 
-	/**
-	 * Holds the name of the device model.
-	 */
+	/** Holds the name of the device model. */
 	UPROPERTY()
 	FString ModelName;
 
-	/**
-	 * Holds the name of the GPU.
-	 */
+	/** Holds the name of the GPU. */
 	UPROPERTY()
 	FString GPUName;
 
-	/**
-	 * Holds the name of the CPU model.
-	 */
+	/** Holds the name of the CPU model. */
 	UPROPERTY()
 	FString CPUModelName;
 
-	/**
-	 * Holds the amount of RAM this device has in gigabytes.
-	 */
+	/** Holds the amount of RAM this device has in gigabytes. */
 	UPROPERTY()
 	uint32 RAMInGB;
 
-	/**
-	 * Holds the name of the current render mode.
-	 */
+	/** Holds the name of the current render mode. */
 	UPROPERTY()
 	FString RenderModeName;
 
-	/**
-	 * Holds the worker's application session identifier.
-	 */
+	/** Holds the worker's application session identifier. */
 	UPROPERTY()
 	FGuid SessionId;
 
-
 	/**
 	 * Default constructor.
 	 */
@@ -237,19 +200,14 @@ struct FAutomationWorkerRequestTests
 {
 	GENERATED_USTRUCT_BODY()
 
-	/**
-	 * Holds a flag indicating whether the developer directory should be included.
-	 */
+	/** Holds a flag indicating whether the developer directory should be included. */
 	UPROPERTY()
 	bool DeveloperDirectoryIncluded;
 
-	/**
-	 * Holds a flag indicating whether the visual commandlet filter is enabled.
-	 */
+	/** Holds a flag indicating whether the visual commandlet filter is enabled. */
 	UPROPERTY()
 	bool VisualCommandletFilterOn;
 
-
 	/**
 	 * Default constructor.
 	 */
@@ -282,19 +240,14 @@ struct FAutomationWorkerRequestTestsReply
 {
 	GENERATED_USTRUCT_BODY()
 
-	/**
-	 * Holds the test information serialized into a string.
-	 */
+	/** Holds the test information serialized into a string. */
 	UPROPERTY()
 	FString TestInfo;
 
-	/**
-	 * Holds the total number of tests returned.
-	 */
+	/** Holds the total number of tests returned. */
 	UPROPERTY()
 	int32 TotalNumTests;
 
-
 	/**
 	 * Default constructor.
 	 */
@@ -327,35 +280,26 @@ struct FAutomationWorkerRunTests
 {
 	GENERATED_USTRUCT_BODY()
 
-	/**
-	 */
+	/** */
 	UPROPERTY()
 	uint32 ExecutionCount;
 
-	/**
-	 */
+	/** */
 	UPROPERTY()
 	int32 RoleIndex;
 
-	/**
-	 * Holds the name of the test to run.
-	 */
+	/** Holds the name of the test to run. */
 	UPROPERTY()
 	FString TestName;
 
-	/**
-	 * If true, we will save out screenshots for tests that support them.
-	 */
+	/** If true, we will save out screenshots for tests that support them. */
 	UPROPERTY()
 	bool bScreenshotsEnabled;
 
-	/**
-	 * If true, we will not resize screen shots.
-	 */
+	/** If true, we will not resize screen shots. */
 	UPROPERTY()
 	bool bUseFullSizeScreenShots;
 
-
 	/**
 	 * Default constructor.
 	 */
@@ -391,38 +335,31 @@ struct FAutomationWorkerRunTestsReply
 {
 	GENERATED_USTRUCT_BODY()
 
-	/**
-	 */
+	/** */
 	UPROPERTY()
 	float Duration;
 
-	/**
-	 */
+	/** */
 	UPROPERTY()
 	TArray<FString> Errors;
 
-	/**
-	 */
+	/** */
 	UPROPERTY()
 	uint32 ExecutionCount;
 
-	/**
-	 */
+	/** */
 	UPROPERTY()
 	TArray<FString> Logs;
 
-	/**
-	 */
+	/** */
 	UPROPERTY()
 	bool Success;
 
-	/**
-	 */
+	/** */
 	UPROPERTY()
 	FString TestName;
 
-	/**
-	 */
+	/** */
 	UPROPERTY()
 	TArray<FString> Warnings;
 };
@@ -444,12 +381,10 @@ struct FAutomationWorkerRequestNextNetworkCommand
 {
 	GENERATED_USTRUCT_BODY()
 
-	/**
-	 */
+	/** */
 	UPROPERTY()
 	uint32 ExecutionCount;
 
-
 	/**
 	 * Default constructor.
 	 */
@@ -490,6 +425,7 @@ struct TStructOpsTypeTraits<FAutomationWorkerNextNetworkCommandReply> : public T
 	};
 };
 
+
 /**
  * Implements a message that is sent in containing a screen shot run during performance test.
  */
@@ -498,11 +434,11 @@ struct FAutomationWorkerScreenImage
 {
 	GENERATED_USTRUCT_BODY()
 
-	// The screen shot data
+	/** The screen shot data. */
 	UPROPERTY()
 	TArray<uint8> ScreenImage;
 
-	// The screen shot name
+	/** The screen shot name. */
 	UPROPERTY()
 	FString ScreenShotName;
 };
@@ -516,6 +452,7 @@ struct TStructOpsTypeTraits<FAutomationWorkerScreenImage> : public TStructOpsTyp
 	};
 };
 
+
 /* Dummy class
  *****************************************************************************/
 
diff --git a/Engine/Source/Runtime/AutomationMessages/Private/AutomationMessagesModule.cpp b/Engine/Source/Runtime/AutomationMessages/Private/AutomationMessagesModule.cpp
--- a/Engine/Source/Runtime/AutomationMessages/Private/AutomationMessagesModule.cpp
+++ b/Engine/Source/Runtime/AutomationMessages/Private/AutomationMessagesModule.cpp
@@ -1,13 +1,8 @@
 // Copyright 1998-2014 Epic Games, Inc. All Rights Reserved.
 
-/*=============================================================================
-	AutomationMessagesModule.cpp: Implements the FAutomationMessagesModule class.
-=============================================================================*/
-
 #include "AutomationMessagesPrivatePCH.h"
 
 
-
 /**
  * Implements the AutomationMessages module.
  */
diff --git a/Engine/Source/Runtime/AutomationMessages/Private/AutomationMessagesPrivatePCH.h b/Engine/Source/Runtime/AutomationMessages/Private/AutomationMessagesPrivatePCH.h
--- a/Engine/Source/Runtime/AutomationMessages/Private/AutomationMessagesPrivatePCH.h
+++ b/Engine/Source/Runtime/AutomationMessages/Private/AutomationMessagesPrivatePCH.h
@@ -1,24 +1,11 @@
 // Copyright 1998-2014 Epic Games, Inc. All Rights Reserved.
 
-/*=============================================================================
-	AutomationMessagesPrivatePCH.h: Pre-compiled header file for the AutomationMessages module.
-=============================================================================*/
+#pragma once
 
-#ifndef AUTOMATIONMESSAGES_PRIVATEPCH_H
-#define AUTOMATIONMESSAGES_PRIVATEPCH_H
-
-
-#include "../Public/AutomationMessages.h"
+#include "AutomationMessages.h"
 
 
 /* Dependencies
  *****************************************************************************/
 
 #include "ModuleManager.h"
-
-
-/* Private includes
- *****************************************************************************/
-
-
-#endif
diff --git a/Engine/Source/Runtime/AutomationMessages/Public/AutomationMessages.h b/Engine/Source/Runtime/AutomationMessages/Public/AutomationMessages.h
--- a/Engine/Source/Runtime/AutomationMessages/Public/AutomationMessages.h
+++ b/Engine/Source/Runtime/AutomationMessages/Public/AutomationMessages.h
@@ -1,9 +1,5 @@
 // Copyright 1998-2014 Epic Games, Inc. All Rights Reserved.
 
-/*=============================================================================
-	AutomationMessages.h: AutomationMessages module public header file.
-=============================================================================*/
-
 #pragma once
 
 
diff --git a/Engine/Source/Runtime/AutomationWorker/Private/AutomationWorkerModule.cpp b/Engine/Source/Runtime/AutomationWorker/Private/AutomationWorkerModule.cpp
--- a/Engine/Source/Runtime/AutomationWorker/Private/AutomationWorkerModule.cpp
+++ b/Engine/Source/Runtime/AutomationWorker/Private/AutomationWorkerModule.cpp
@@ -1,15 +1,10 @@
 // Copyright 1998-2014 Epic Games, Inc. All Rights Reserved.
 
-/*=============================================================================
-	AutomationWorkerModule.cpp: Implements the FAutomationWorkerModule class.
-=============================================================================*/
-
 #include "AutomationWorkerPrivatePCH.h"
 
 
 #define LOCTEXT_NAMESPACE "AutomationTest"
 
-
 IMPLEMENT_MODULE(FAutomationWorkerModule, AutomationWorker);
 
 
@@ -152,13 +147,13 @@ void FAutomationWorkerModule::ReportNetworkCommandComplete()
 	}
 }
 
+
 /**
  *	Takes a large transport array and splits it into pieces of a desired size and returns the portion of this which is requested
  *
- * @param FullTransportArray	- The whole series of data
- * @param CurrentChunkIndex		- The The chunk we are requesting
- * @param NumToSend				- The maximum number of bytes we should be splitting into.
- *
+ * @param FullTransportArray The whole series of data
+ * @param CurrentChunkIndex The The chunk we are requesting
+ * @param NumToSend The maximum number of bytes we should be splitting into.
  * @return The section of the transport array which matches our index requested
  */
 TArray< uint8 > GetTransportSection( const TArray< uint8 >& FullTransportArray, const int32 NumToSend, const int32 RequestedChunkIndex )
@@ -292,7 +287,7 @@ void FAutomationWorkerModule::HandlePingMessage( const FAutomationWorkerPing& Me
 	MessageEndpoint->Send(new FAutomationWorkerPong(), Context->GetSender());
 }
 
-// Handles FAutomationWorkerResetTests messages.
+
 void FAutomationWorkerModule::HandleResetTests( const FAutomationWorkerResetTests& Message, const IMessageContextRef& Context )
 {
 	FAutomationTestFramework::GetInstance().ResetTests();
@@ -366,6 +361,7 @@ void FAutomationWorkerModule::HandleScreenShotCaptured( int32 Width, int32 Heigh
 }
 #endif
 
+
 void FAutomationWorkerModule::HandleRunTestsMessage( const FAutomationWorkerRunTests& Message, const IMessageContextRef& Context )
 {
 	ExecutionCount = Message.ExecutionCount;
@@ -382,6 +378,7 @@ void FAutomationWorkerModule::HandleRunTestsMessage( const FAutomationWorkerRunT
 	FAutomationTestFramework::GetInstance().StartTestByName(Message.TestName, Message.RoleIndex);
 }
 
+
 void FAutomationWorkerModule::RunTest(const FString& InTestToRun, const int32 InRoleIndex, FStopTestEvent const& InStopTestEvent)
 {
 	TestName = InTestToRun;
@@ -396,8 +393,10 @@ void FAutomationWorkerModule::RunTest(const FString& InTestToRun, const int32 In
 	FAutomationTestFramework::GetInstance().StartTestByName(InTestToRun, InRoleIndex);
 }
 
-// This is sortof a local controller to run tests and spew results, mostly used by automated testing
 
+/**
+ * Implements a local controller to run tests and spew results, mostly used by automated testing.
+ */
 static struct FQueueTests
 {
 	struct FJob
@@ -423,8 +422,7 @@ static struct FQueueTests
 		: NumTestsRun(0)
 		, bTestInPogress(false)
 		, bTicking(false)
-	{
-	}
+	{ }
 
 	void NewTest(FString const& Command, int32 RoleIndex = 0, FOutputDevice* Ar = GLog)
 	{
@@ -505,6 +503,7 @@ static struct FQueueTests
 
 } QueueTests;
 
+
 bool DirectAutomationCommand(const TCHAR* Cmd, FOutputDevice* Ar = GLog)
 {
 	bool bResult = false;
@@ -553,7 +552,6 @@ bool DirectAutomationCommand(const TCHAR* Cmd, FOutputDevice* Ar = GLog)
 }
 
 
-
 static class FAutomationTestCmd : private FSelfRegisteringExec
 {
 public:
@@ -565,6 +563,4 @@ public:
 } AutomationTestCmd;
 
 
-
-
 #undef LOCTEXT_NAMESPACE
diff --git a/Engine/Source/Runtime/AutomationWorker/Private/AutomationWorkerModule.h b/Engine/Source/Runtime/AutomationWorker/Private/AutomationWorkerModule.h
--- a/Engine/Source/Runtime/AutomationWorker/Private/AutomationWorkerModule.h
+++ b/Engine/Source/Runtime/AutomationWorker/Private/AutomationWorkerModule.h
@@ -1,9 +1,5 @@
 // Copyright 1998-2014 Epic Games, Inc. All Rights Reserved.
 
-/*=============================================================================
-	AutomationWorkerModule.h: Declares the FAutomationWorkerModule class.
-=============================================================================*/
-
 #pragma once
 
 
@@ -63,7 +59,7 @@ protected:
 	/** 
 	 * Send a list of all the tests supported by the worker
 	 *
-	 * @param ControllerAddress - The message address of the controller that requested the tests.
+	 * @param ControllerAddress The message address of the controller that requested the tests.
 	 */
 	void SendTests( const FMessageAddress& ControllerAddress );
 
@@ -92,13 +88,13 @@ private:
 
 #if WITH_ENGINE
 	/** Invoked when we have screen shot to send. */
-	void HandleScreenShotCaptured(int32 Width, int32 Height, const TArray<FColor>& Bitmap, const FString& ScreenShotName);
+	void HandleScreenShotCaptured( int32 Width, int32 Height, const TArray<FColor>& Bitmap, const FString& ScreenShotName );
 #endif
 
 private:
 
 	// The collection of test data we are to send to a controller
-	TArray< FAutomationTestInfo > TestInfo;
+	TArray<FAutomationTestInfo> TestInfo;
 
 private:
 
diff --git a/Engine/Source/Runtime/AutomationWorker/Private/AutomationWorkerPrivatePCH.h b/Engine/Source/Runtime/AutomationWorker/Private/AutomationWorkerPrivatePCH.h
--- a/Engine/Source/Runtime/AutomationWorker/Private/AutomationWorkerPrivatePCH.h
+++ b/Engine/Source/Runtime/AutomationWorker/Private/AutomationWorkerPrivatePCH.h
@@ -1,12 +1,7 @@
 // Copyright 1998-2014 Epic Games, Inc. All Rights Reserved.
 
-/*=============================================================================
-	AutomationWorkerPrivatePCH.h: Pre-compiled header file for the AutomationWorker module.
-=============================================================================*/
-
 #pragma once
 
-
 #include "AutomationWorker.h"
 
 
@@ -22,6 +17,7 @@
 	#include "AutomationCommon.h"
 #endif
 
+
 /* Private includes
  *****************************************************************************/
 
diff --git a/Engine/Source/Runtime/AutomationWorker/Public/AutomationWorker.h b/Engine/Source/Runtime/AutomationWorker/Public/AutomationWorker.h
--- a/Engine/Source/Runtime/AutomationWorker/Public/AutomationWorker.h
+++ b/Engine/Source/Runtime/AutomationWorker/Public/AutomationWorker.h
@@ -1,9 +1,5 @@
 // Copyright 1998-2014 Epic Games, Inc. All Rights Reserved.
 
-/*=============================================================================
-	AutomationWorker.h: Automation Worker module public header file.
-=============================================================================*/
-
 #pragma once
 
 
diff --git a/Engine/Source/Runtime/AutomationWorker/Public/Interfaces/IAutomationWorkerModule.h b/Engine/Source/Runtime/AutomationWorker/Public/Interfaces/IAutomationWorkerModule.h
--- a/Engine/Source/Runtime/AutomationWorker/Public/Interfaces/IAutomationWorkerModule.h
+++ b/Engine/Source/Runtime/AutomationWorker/Public/Interfaces/IAutomationWorkerModule.h
@@ -1,9 +1,5 @@
 // Copyright 1998-2014 Epic Games, Inc. All Rights Reserved.
 
-/*=============================================================================
-	IAutomationWorkerModule.h: Declares the IAutomationWorkerModule interface.
-=============================================================================*/
-
 #pragma once
 
 
diff --git a/Engine/Source/Runtime/Core/Private/Containers/ContainerAllocationPolicies.cpp b/Engine/Source/Runtime/Core/Private/Containers/ContainerAllocationPolicies.cpp
--- a/Engine/Source/Runtime/Core/Private/Containers/ContainerAllocationPolicies.cpp
+++ b/Engine/Source/Runtime/Core/Private/Containers/ContainerAllocationPolicies.cpp
@@ -1,9 +1,5 @@
 // Copyright 1998-2014 Epic Games, Inc. All Rights Reserved.
 
-/*=============================================================================
-	ContainerAllocationPolicies.cpp: Container allocation policy implementation.
-=============================================================================*/
-
 #include "CorePrivate.h"
 
 #if TRACK_ARRAY_SLACK 
@@ -12,6 +8,7 @@
 
 FStackTracker* GSlackTracker = NULL;
 
+
 /** Updates an existing call stack trace with new data for this particular call*/
 static void SlackTrackerUpdateFn(const FStackTracker::FCallStack& CallStack, void* UserData)
 {
diff --git a/Engine/Source/Runtime/Core/Private/Containers/ContainersTest.cpp b/Engine/Source/Runtime/Core/Private/Containers/ContainersTest.cpp
--- a/Engine/Source/Runtime/Core/Private/Containers/ContainersTest.cpp
+++ b/Engine/Source/Runtime/Core/Private/Containers/ContainersTest.cpp
@@ -1,17 +1,15 @@
 // Copyright 1998-2014 Epic Games, Inc. All Rights Reserved.
 
-/*=============================================================================
-	ContainersTest.cpp: Unit test for containers.
-=============================================================================*/
-
 #include "CorePrivate.h"
 #include "AutomationTest.h"
 
+
 IMPLEMENT_SIMPLE_AUTOMATION_TEST(FContainersTest, "Core.Misc.Containers", EAutomationTestFlags::ATF_SmokeTest)
 
 #define MAX_TEST_OBJECTS      65
 #define MAX_TEST_OBJECTS_STEP 1
 
+
 namespace
 {
 	namespace EContainerTestType
diff --git a/Engine/Source/Runtime/Core/Private/Containers/StackTracker.cpp b/Engine/Source/Runtime/Core/Private/Containers/StackTracker.cpp
--- a/Engine/Source/Runtime/Core/Private/Containers/StackTracker.cpp
+++ b/Engine/Source/Runtime/Core/Private/Containers/StackTracker.cpp
@@ -1,14 +1,12 @@
 // Copyright 1998-2014 Epic Games, Inc. All Rights Reserved.
 
-/*=============================================================================
-StackTracker.cpp: Stack Tracking within Unreal Engine.
-=============================================================================*/
 #include "CorePrivate.h"
-
 #include "StackTracker.h"
 
+
 DEFINE_LOG_CATEGORY_STATIC(LogStackTracker, Log, All);
 
+
 /**
  * Captures the current stack and updates stack tracking information.
  * optionally stores a user data pointer that the tracker will take ownership of and delete upon reset
diff --git a/Engine/Source/Runtime/Core/Private/Containers/String.cpp b/Engine/Source/Runtime/Core/Private/Containers/String.cpp
--- a/Engine/Source/Runtime/Core/Private/Containers/String.cpp
+++ b/Engine/Source/Runtime/Core/Private/Containers/String.cpp
@@ -1,14 +1,10 @@
 // Copyright 1998-2014 Epic Games, Inc. All Rights Reserved.
 
-/*=============================================================================
-	String.cpp: Implementation of FString and FString
-=============================================================================*/
-
 #include "CorePrivate.h"
 
-/*-----------------------------------------------------------------------------
-	FString implementation.
------------------------------------------------------------------------------*/
+
+/* FString implementation
+ *****************************************************************************/
 
 void FString::TrimToNullTerminator()
 {
diff --git a/Engine/Source/Runtime/Core/Private/GenericPlatform/GenericPlatformFile.cpp b/Engine/Source/Runtime/Core/Private/GenericPlatform/GenericPlatformFile.cpp
--- a/Engine/Source/Runtime/Core/Private/GenericPlatform/GenericPlatformFile.cpp
+++ b/Engine/Source/Runtime/Core/Private/GenericPlatform/GenericPlatformFile.cpp
@@ -1,12 +1,9 @@
 // Copyright 1998-2014 Epic Games, Inc. All Rights Reserved.
 
-/*=============================================================================
-	GenericPlatformFile.cpp: Generic implementations of platform file I/O functions
-=============================================================================*/
-
 #include "CorePrivate.h"
 #include "ModuleManager.h"
 
+
 int64 IFileHandle::Size()
 {
 	int64 Current = Tell();
diff --git a/Engine/Source/Runtime/Core/Private/GenericPlatform/GenericPlatformMath.cpp b/Engine/Source/Runtime/Core/Private/GenericPlatform/GenericPlatformMath.cpp
--- a/Engine/Source/Runtime/Core/Private/GenericPlatform/GenericPlatformMath.cpp
+++ b/Engine/Source/Runtime/Core/Private/GenericPlatform/GenericPlatformMath.cpp
@@ -1,14 +1,11 @@
 // Copyright 1998-2014 Epic Games, Inc. All Rights Reserved.
 
-/*=============================================================================
-	GenericPlatformMath.cpp: Generic implementations of Math platform functions
-=============================================================================*/
-
 #include "CorePrivate.h"
 #include "BigInt.h"
 
 static int32 GSRandSeed;
 
+
 void FGenericPlatformMath::SRandInit( int32 Seed ) 
 {
 	GSRandSeed = Seed; 
diff --git a/Engine/Source/Runtime/Core/Private/GenericPlatform/GenericPlatformMemory.cpp b/Engine/Source/Runtime/Core/Private/GenericPlatform/GenericPlatformMemory.cpp
--- a/Engine/Source/Runtime/Core/Private/GenericPlatform/GenericPlatformMemory.cpp
+++ b/Engine/Source/Runtime/Core/Private/GenericPlatform/GenericPlatformMemory.cpp
@@ -1,9 +1,5 @@
 // Copyright 1998-2014 Epic Games, Inc. All Rights Reserved.
 
-/*=============================================================================
-	GenericPlatformMemory.cpp: Generic implementations of memory platform functions
-=============================================================================*/
-
 #include "CorePrivate.h"
 #include "MallocAnsi.h"
 #include "GenericPlatformMemoryPoolStats.h"
diff --git a/Engine/Source/Runtime/Core/Private/GenericPlatform/GenericPlatformMisc.cpp b/Engine/Source/Runtime/Core/Private/GenericPlatform/GenericPlatformMisc.cpp
--- a/Engine/Source/Runtime/Core/Private/GenericPlatform/GenericPlatformMisc.cpp
+++ b/Engine/Source/Runtime/Core/Private/GenericPlatform/GenericPlatformMisc.cpp
@@ -1,9 +1,5 @@
 ﻿// Copyright 1998-2014 Epic Games, Inc. All Rights Reserved.
 
-/*=============================================================================
-	GenericPlatformMisc.cpp: Generic implementations of misc platform functions
-=============================================================================*/
-
 #include "CorePrivate.h"
 #include "MallocAnsi.h"
 #include "GenericApplication.h"
@@ -18,7 +14,7 @@
 #include "UProjectInfo.h"
 
 #if UE_ENABLE_ICU
-#include <unicode/locid.h>
+	#include <unicode/locid.h>
 #endif
 
 DEFINE_LOG_CATEGORY_STATIC(LogGenericPlatformMisc, Log, All);
@@ -26,6 +22,7 @@ DEFINE_LOG_CATEGORY_STATIC(LogGenericPlatformMisc, Log, All);
 /** Holds an override path if a program has special needs */
 FString OverrideGameDir;
 
+
 /* EBuildConfigurations interface
  *****************************************************************************/
 
@@ -57,7 +54,6 @@ namespace EBuildConfigurations
 		return Unknown;
 	}
 
-
 	const TCHAR* ToString( EBuildConfigurations::Type Configuration )
 	{
 		switch (Configuration)
@@ -131,7 +127,6 @@ namespace EBuildTargets
 		return Unknown;
 	}
 
-
 	const TCHAR* ToString( EBuildTargets::Type Target )
 	{
 		switch (Target)
@@ -154,11 +149,13 @@ namespace EBuildTargets
 
 /* FGenericPlatformMisc interface
  *****************************************************************************/
+
 #if !UE_BUILD_SHIPPING
 	bool FGenericPlatformMisc::bShouldPromptForRemoteDebugging = false;
 	bool FGenericPlatformMisc::bPromptForRemoteDebugOnEnsure = false;
 #endif	//#if !UE_BUILD_SHIPPING
 
+
 GenericApplication* FGenericPlatformMisc::CreateApplication()
 {
 	return new GenericApplication( nullptr );
diff --git a/Engine/Source/Runtime/Core/Private/GenericPlatform/GenericPlatformNamedPipe.cpp b/Engine/Source/Runtime/Core/Private/GenericPlatform/GenericPlatformNamedPipe.cpp
--- a/Engine/Source/Runtime/Core/Private/GenericPlatform/GenericPlatformNamedPipe.cpp
+++ b/Engine/Source/Runtime/Core/Private/GenericPlatform/GenericPlatformNamedPipe.cpp
@@ -1,9 +1,5 @@
 // Copyright 1998-2014 Epic Games, Inc. All Rights Reserved.
 
-/*=============================================================================================
-	GenericPlatformNamedPipe.cpp: Implements the FGenericPlatformNamedPipe class
-==============================================================================================*/
-
 #include "CorePrivate.h"
 
 
@@ -14,6 +10,7 @@ FGenericPlatformNamedPipe::FGenericPlatformNamedPipe()
 	NamePtr = new FString();
 }
 
+
 FGenericPlatformNamedPipe::~FGenericPlatformNamedPipe()
 {
 	delete NamePtr;
diff --git a/Engine/Source/Runtime/Core/Private/GenericPlatform/GenericPlatformOutputDevices.cpp b/Engine/Source/Runtime/Core/Private/GenericPlatform/GenericPlatformOutputDevices.cpp
--- a/Engine/Source/Runtime/Core/Private/GenericPlatform/GenericPlatformOutputDevices.cpp
+++ b/Engine/Source/Runtime/Core/Private/GenericPlatform/GenericPlatformOutputDevices.cpp
@@ -1,12 +1,9 @@
 // Copyright 1998-2014 Epic Games, Inc. All Rights Reserved.
 
-/*=============================================================================
-	GenericPlatformOutputDevices.cpp: Generic implementations of OutputDevices platform functions
-=============================================================================*/
-
 #include "CorePrivate.h"
 #include "FeedbackContextAnsi.h"
 
+
 FString								FGenericPlatformOutputDevices::GetAbsoluteLogFilename()
 {
 	static TCHAR		Filename[1024] = { 0 };
diff --git a/Engine/Source/Runtime/Core/Private/GenericPlatform/GenericPlatformProcess.cpp b/Engine/Source/Runtime/Core/Private/GenericPlatform/GenericPlatformProcess.cpp
--- a/Engine/Source/Runtime/Core/Private/GenericPlatform/GenericPlatformProcess.cpp
+++ b/Engine/Source/Runtime/Core/Private/GenericPlatform/GenericPlatformProcess.cpp
@@ -1,16 +1,14 @@
 // Copyright 1998-2014 Epic Games, Inc. All Rights Reserved.
 
-/*=============================================================================
-	GenericPlatformProcess.cpp: Generic implementations of Process platform functions
-=============================================================================*/
-
 #include "CorePrivate.h"
 
+
 #if PLATFORM_HAS_BSD_TIME 
-#include <unistd.h> // for usleep
-#include <sched.h>  // for sched_yield
+	#include <unistd.h> // for usleep
+	#include <sched.h>  // for sched_yield
 #endif
 
+
 void* FGenericPlatformProcess::GetDllHandle( const TCHAR* Filename )
 {
 	UE_LOG(LogHAL, Fatal, TEXT("FPlatformProcess::GetDllHandle not implemented on this platform"));
diff --git a/Engine/Source/Runtime/Core/Private/GenericPlatform/GenericPlatformProperties.cpp b/Engine/Source/Runtime/Core/Private/GenericPlatform/GenericPlatformProperties.cpp
--- a/Engine/Source/Runtime/Core/Private/GenericPlatform/GenericPlatformProperties.cpp
+++ b/Engine/Source/Runtime/Core/Private/GenericPlatform/GenericPlatformProperties.cpp
@@ -1,14 +1,10 @@
 // Copyright 1998-2014 Epic Games, Inc. All Rights Reserved.
 
-/*=============================================================================
-	GenericPlatformMath.cpp: Generic implementations of Math platform functions
-=============================================================================*/
-
 #include "CorePrivate.h"
 
+
 // This has to be in a .cpp since it needs to come after the FPlatformProperties class has been typedef'd
 const char* FGenericPlatformProperties::IniPlatformName()
 {
 	return FPlatformProperties::PlatformName();
 }
-
diff --git a/Engine/Source/Runtime/Core/Private/GenericPlatform/GenericPlatformStackWalk.cpp b/Engine/Source/Runtime/Core/Private/GenericPlatform/GenericPlatformStackWalk.cpp
--- a/Engine/Source/Runtime/Core/Private/GenericPlatform/GenericPlatformStackWalk.cpp
+++ b/Engine/Source/Runtime/Core/Private/GenericPlatform/GenericPlatformStackWalk.cpp
@@ -1,11 +1,8 @@
 // Copyright 1998-2014 Epic Games, Inc. All Rights Reserved.
 
-/*=============================================================================
-	GenericPlatformStackWalk.cpp: Generic, do nothing stack walk implementation
-=============================================================================*/
-
 #include "CorePrivate.h"
 
+
 bool FGenericPlatformStackWalk::ProgramCounterToHumanReadableString( int32 CurrentCallDepth, uint64 ProgramCounter, ANSICHAR* HumanReadableString, SIZE_T HumanReadableStringSize, EVerbosityFlags VerbosityFlags, FGenericCrashContext* Context )
 {
 	if (HumanReadableString && HumanReadableStringSize > 0)
diff --git a/Engine/Source/Runtime/Core/Private/GenericPlatform/GenericPlatformString.cpp b/Engine/Source/Runtime/Core/Private/GenericPlatform/GenericPlatformString.cpp
--- a/Engine/Source/Runtime/Core/Private/GenericPlatform/GenericPlatformString.cpp
+++ b/Engine/Source/Runtime/Core/Private/GenericPlatform/GenericPlatformString.cpp
@@ -1,17 +1,15 @@
 // Copyright 1998-2014 Epic Games, Inc. All Rights Reserved.
 
-/*=============================================================================
-	GenericPlatformString.cpp: Generic implementations of string functions
-=============================================================================*/
-
 #include "CorePrivate.h"
 
+
 DEFINE_LOG_CATEGORY_STATIC(LogGenericPlatformString, Log, All);
 
 template <> const TCHAR* FGenericPlatformString::GetEncodingTypeName<ANSICHAR>() { return TEXT("ANSICHAR"); }
 template <> const TCHAR* FGenericPlatformString::GetEncodingTypeName<WIDECHAR>() { return TEXT("WIDECHAR"); }
 template <> const TCHAR* FGenericPlatformString::GetEncodingTypeName<UCS2CHAR>() { return TEXT("UCS2CHAR"); }
 
+
 void* FGenericPlatformString::Memcpy(void* Dest, const void* Src, SIZE_T Count)
 {
 	return FMemory::Memcpy(Dest, Src, Count);
diff --git a/Engine/Source/Runtime/Core/Private/GenericPlatform/GenericPlatformTime.cpp b/Engine/Source/Runtime/Core/Private/GenericPlatform/GenericPlatformTime.cpp
--- a/Engine/Source/Runtime/Core/Private/GenericPlatform/GenericPlatformTime.cpp
+++ b/Engine/Source/Runtime/Core/Private/GenericPlatform/GenericPlatformTime.cpp
@@ -1,9 +1,5 @@
 // Copyright 1998-2014 Epic Games, Inc. All Rights Reserved.
 
-/*=============================================================================
-	GenericPlatformTime.cpp: Generic implementations of Time platform functions
-=============================================================================*/
-
 #include "CorePrivate.h"
 
 DEFINE_LOG_CATEGORY_STATIC(LogGenericPlatformTime, Log, All);
diff --git a/Engine/Source/Runtime/Core/Private/GenericPlatform/StandardPlatformString.cpp b/Engine/Source/Runtime/Core/Private/GenericPlatform/StandardPlatformString.cpp
--- a/Engine/Source/Runtime/Core/Private/GenericPlatform/StandardPlatformString.cpp
+++ b/Engine/Source/Runtime/Core/Private/GenericPlatform/StandardPlatformString.cpp
@@ -1,21 +1,19 @@
 // Copyright 1998-2014 Epic Games, Inc. All Rights Reserved.
 
-/*=============================================================================
- StandardPlatformString.cpp: Standard implementations of string functions
- =============================================================================*/
-
 #include "CorePrivate.h"
 
+
 DEFINE_LOG_CATEGORY_STATIC(LogStandardPlatformString, Log, All);
 
 #if !PLATFORM_USE_SYSTEM_VSWPRINTF
 
 #if PLATFORM_IOS
-#define VA_LIST_REF va_list&
+	#define VA_LIST_REF va_list&
 #else
-#define VA_LIST_REF va_list
+	#define VA_LIST_REF va_list
 #endif
 
+
 struct FFormatInfo
 {
 	WIDECHAR Format[32];
diff --git a/Engine/Source/Runtime/Core/Private/Misc/MonitoredProcess.cpp b/Engine/Source/Runtime/Core/Private/Misc/MonitoredProcess.cpp
--- a/Engine/Source/Runtime/Core/Private/Misc/MonitoredProcess.cpp
+++ b/Engine/Source/Runtime/Core/Private/Misc/MonitoredProcess.cpp
@@ -1,9 +1,5 @@
 // Copyright 1998-2014 Epic Games, Inc. All Rights Reserved.
 
-/*=============================================================================================
-	MonitoredProcess.cpp: Implements the FMonitoredProcess class
-==============================================================================================*/
-
 #include "CorePrivate.h"
 
 
diff --git a/Engine/Source/Runtime/Core/Public/Containers/EnumAsByte.h b/Engine/Source/Runtime/Core/Public/Containers/EnumAsByte.h
--- a/Engine/Source/Runtime/Core/Public/Containers/EnumAsByte.h
+++ b/Engine/Source/Runtime/Core/Public/Containers/EnumAsByte.h
@@ -1,9 +1,5 @@
 // Copyright 1998-2014 Epic Games, Inc. All Rights Reserved.
 
-/*=============================================================================
-	EnumAsByte.h: FEnumAsByte template
-=============================================================================*/
-
 #pragma once
 
 
@@ -24,7 +20,7 @@ public:
 	/**
 	 * Copy constructor.
 	 *
-	 * @param InValue value to construct with 
+	 * @param InValue value to construct with.
 	 */
 	FORCEINLINE TEnumAsByte( const TEnumAsByte &InValue )
 		: Value(InValue.Value)
@@ -33,7 +29,7 @@ public:
 	/**
 	 * Constructor, initialize to the enum value.
 	 *
-	 * @param InValue value to construct with .
+	 * @param InValue value to construct with.
 	 */
 	FORCEINLINE TEnumAsByte( TEnum InValue )
 		: Value(InValue)
@@ -42,14 +38,14 @@ public:
 	/**
 	 * Constructor, initialize to the int32 value.
 	 *
-	 * @param InValue value to construct with 
+	 * @param InValue value to construct with.
 	 */
 	explicit FORCEINLINE TEnumAsByte( int32 InValue )
 		: Value(InValue)
 	{ }
 
 	/**
-	 * Constructor, initialize to the int32 value
+	 * Constructor, initialize to the int32 value.
 	 *
 	 * @param InValue value to construct with.
 	 */
diff --git a/Engine/Source/Runtime/Core/Public/Containers/List.h b/Engine/Source/Runtime/Core/Public/Containers/List.h
--- a/Engine/Source/Runtime/Core/Public/Containers/List.h
+++ b/Engine/Source/Runtime/Core/Public/Containers/List.h
@@ -1,11 +1,8 @@
 // Copyright 1998-2014 Epic Games, Inc. All Rights Reserved.
 
-/*=============================================================================
-	List.h: Dynamic list definitions.
-=============================================================================*/
-
 #pragma once
 
+
 /**
  * Encapsulates a link in a single linked list with constant access time.
  */
@@ -108,7 +105,7 @@ public:
 	/**
 	 * Adds this element to a list, before the given element.
 	 *
-	 * @param Before	The link to insert this element before.
+	 * @param Before The link to insert this element before.
 	 */
 	void Link(TLinkedList*& Before)
 	{
@@ -158,6 +155,7 @@ private:
 	TLinkedList** PrevLink;
 };
 
+
 /**
  * Double linked list.
  */
@@ -518,6 +516,7 @@ private:
 	TDoubleLinkedList& operator=(const TDoubleLinkedList&);
 };
 
+
 /*----------------------------------------------------------------------------
 	TList.
 ----------------------------------------------------------------------------*/
@@ -540,4 +539,3 @@ public:
 		Next = InNext;
 	}
 };
-
diff --git a/Engine/Source/Runtime/Core/Public/Containers/LockFreeFixedSizeAllocator.h b/Engine/Source/Runtime/Core/Public/Containers/LockFreeFixedSizeAllocator.h
--- a/Engine/Source/Runtime/Core/Public/Containers/LockFreeFixedSizeAllocator.h
+++ b/Engine/Source/Runtime/Core/Public/Containers/LockFreeFixedSizeAllocator.h
@@ -1,9 +1,5 @@
 // Copyright 1998-2014 Epic Games, Inc. All Rights Reserved.
 
-/*=============================================================================
-	LockFreeFixedSizeAllocator.h: A lock free pooled allocator for fixed size blocks of a particular type
-=============================================================================*/
-
 #pragma once
 
 #include "LockFreeList.h"
@@ -13,6 +9,7 @@
 
 #define USE_RECYCLING (1)
 
+
 #if !USE_RECYCLING
 
 /** Thread safe, lock free pooling allocator of fixed size blocks that never returns free space until program shutdown. **/
@@ -33,7 +30,9 @@ public:
 };
 
 
-#else
+#else //!USE_RECYCLING
+
+
 /** Thread safe, lock free pooling allocator of fixed size blocks that never returns free space until program shutdown. **/
 template<int32 SIZE>
 class TLockFreeFixedSizeAllocator	// alignment isn't handled, assumes FMemory::Malloc will work
@@ -86,7 +85,9 @@ private:
 	FNoopCounter		NumFree;
 #endif
 };
-#endif
+
+#endif //!USE_RECYCLING
+
 
 /** Thread safe, lock free pooling allocator of memory for instances of T. Never returns free space until program shutdown. **/
 template<class T>
@@ -110,4 +111,3 @@ public:
 		TLockFreeFixedSizeAllocator<sizeof(T)>::Free(Item);
 	}
 };
-
diff --git a/Engine/Source/Runtime/Core/Public/Containers/LockFreeList.h b/Engine/Source/Runtime/Core/Public/Containers/LockFreeList.h
--- a/Engine/Source/Runtime/Core/Public/Containers/LockFreeList.h
+++ b/Engine/Source/Runtime/Core/Public/Containers/LockFreeList.h
@@ -1,11 +1,8 @@
 // Copyright 1998-2014 Epic Games, Inc. All Rights Reserved.
 
-/*=============================================================================
-	LockFreeList.h: A lock free single linked list
-=============================================================================*/
-
 #pragma once
 
+
 CORE_API DECLARE_LOG_CATEGORY_EXTERN(LogLockFreeList, Log, All);
 
 // what level of checking to perform...normally checkLockFreePointerList but could be ensure or check
