Enabled UHT to digest GENERATED_BODY instead of GENERATED_UCLASS_BODY, GENERATED_USTRUCT_BODY, GENERATED_UINTERFACE_BODY or GENERATED_IINTERFACE_BODY, changed every occurence to the new syntax and fix

diff --git a/Engine/Source/Runtime/AIModule/Classes/BehaviorTree/Tasks/BTTask_PushPawnAction.h b/Engine/Source/Runtime/AIModule/Classes/BehaviorTree/Tasks/BTTask_PushPawnAction.h
--- a/Engine/Source/Runtime/AIModule/Classes/BehaviorTree/Tasks/BTTask_PushPawnAction.h
+++ b/Engine/Source/Runtime/AIModule/Classes/BehaviorTree/Tasks/BTTask_PushPawnAction.h
@@ -14,7 +14,9 @@ class UPawnAction;
 UCLASS()
 class AIMODULE_API UBTTask_PushPawnAction : public UBTTask_PawnActionBase
 {
-	GENERATED_UCLASS_BODY()
+	GENERATED_BODY()
+public:
+	UBTTask_PushPawnAction(const FObjectInitializer& ObjectInitializer = FObjectInitializer::Get());
 
 protected:
 	UPROPERTY(EditAnywhere, Instanced, Category = Action)
diff --git a/Engine/Source/Runtime/AIModule/Classes/BehaviorTree/Tasks/BTTask_RotateToFaceBBEntry.h b/Engine/Source/Runtime/AIModule/Classes/BehaviorTree/Tasks/BTTask_RotateToFaceBBEntry.h
--- a/Engine/Source/Runtime/AIModule/Classes/BehaviorTree/Tasks/BTTask_RotateToFaceBBEntry.h
+++ b/Engine/Source/Runtime/AIModule/Classes/BehaviorTree/Tasks/BTTask_RotateToFaceBBEntry.h
@@ -12,7 +12,9 @@
 UCLASS(config = Game)
 class AIMODULE_API UBTTask_RotateToFaceBBEntry : public UBTTask_BlackboardBase
 {
-	GENERATED_UCLASS_BODY()
+	GENERATED_BODY()
+public:
+	UBTTask_RotateToFaceBBEntry(const FObjectInitializer& ObjectInitializer = FObjectInitializer::Get());
 
 protected:
 	/** Success condition precision in decrees */
diff --git a/Engine/Source/Runtime/AIModule/Classes/BehaviorTree/Tasks/BTTask_RunBehavior.h b/Engine/Source/Runtime/AIModule/Classes/BehaviorTree/Tasks/BTTask_RunBehavior.h
--- a/Engine/Source/Runtime/AIModule/Classes/BehaviorTree/Tasks/BTTask_RunBehavior.h
+++ b/Engine/Source/Runtime/AIModule/Classes/BehaviorTree/Tasks/BTTask_RunBehavior.h
@@ -21,7 +21,9 @@
 UCLASS()
 class AIMODULE_API UBTTask_RunBehavior : public UBTTaskNode
 {
-	GENERATED_UCLASS_BODY()
+	GENERATED_BODY()
+public:
+	UBTTask_RunBehavior(const FObjectInitializer& ObjectInitializer = FObjectInitializer::Get());
 
 	virtual EBTNodeResult::Type ExecuteTask(UBehaviorTreeComponent& OwnerComp, uint8* NodeMemory) override;
 	virtual uint16 GetInstanceMemorySize() const override;
diff --git a/Engine/Source/Runtime/AIModule/Classes/BehaviorTree/Tasks/BTTask_RunEQSQuery.h b/Engine/Source/Runtime/AIModule/Classes/BehaviorTree/Tasks/BTTask_RunEQSQuery.h
--- a/Engine/Source/Runtime/AIModule/Classes/BehaviorTree/Tasks/BTTask_RunEQSQuery.h
+++ b/Engine/Source/Runtime/AIModule/Classes/BehaviorTree/Tasks/BTTask_RunEQSQuery.h
@@ -19,7 +19,9 @@ struct FBTEnvQueryTaskMemory
 UCLASS()
 class AIMODULE_API UBTTask_RunEQSQuery : public UBTTask_BlackboardBase
 {
-	GENERATED_UCLASS_BODY()
+	GENERATED_BODY()
+public:
+	UBTTask_RunEQSQuery(const FObjectInitializer& ObjectInitializer = FObjectInitializer::Get());
 
 	/** query to run */
 	UPROPERTY(Category=Node, EditAnywhere)
diff --git a/Engine/Source/Runtime/AIModule/Classes/BehaviorTree/Tasks/BTTask_SetTagCooldown.h b/Engine/Source/Runtime/AIModule/Classes/BehaviorTree/Tasks/BTTask_SetTagCooldown.h
--- a/Engine/Source/Runtime/AIModule/Classes/BehaviorTree/Tasks/BTTask_SetTagCooldown.h
+++ b/Engine/Source/Runtime/AIModule/Classes/BehaviorTree/Tasks/BTTask_SetTagCooldown.h
@@ -10,7 +10,9 @@
 UCLASS()
 class AIMODULE_API UBTTask_SetTagCooldown : public UBTTaskNode
 {
-	GENERATED_UCLASS_BODY()
+	GENERATED_BODY()
+public:
+	UBTTask_SetTagCooldown(const FObjectInitializer& ObjectInitializer = FObjectInitializer::Get());
 
 	/** gameplay tag that will be used for the cooldown */
 	UPROPERTY(Category = Cooldown, EditAnywhere)
diff --git a/Engine/Source/Runtime/AIModule/Classes/BehaviorTree/Tasks/BTTask_Wait.h b/Engine/Source/Runtime/AIModule/Classes/BehaviorTree/Tasks/BTTask_Wait.h
--- a/Engine/Source/Runtime/AIModule/Classes/BehaviorTree/Tasks/BTTask_Wait.h
+++ b/Engine/Source/Runtime/AIModule/Classes/BehaviorTree/Tasks/BTTask_Wait.h
@@ -16,7 +16,9 @@ struct FBTWaitTaskMemory
 UCLASS()
 class AIMODULE_API UBTTask_Wait : public UBTTaskNode
 {
-	GENERATED_UCLASS_BODY()
+	GENERATED_BODY()
+public:
+	UBTTask_Wait(const FObjectInitializer& ObjectInitializer = FObjectInitializer::Get());
 
 	/** wait time in seconds */
 	UPROPERTY(Category = Wait, EditAnywhere, meta = (ClampMin = "0.0", UIMin = "0.0"))
diff --git a/Engine/Source/Runtime/AIModule/Classes/BehaviorTree/Tasks/BTTask_WaitBlackboardTime.h b/Engine/Source/Runtime/AIModule/Classes/BehaviorTree/Tasks/BTTask_WaitBlackboardTime.h
--- a/Engine/Source/Runtime/AIModule/Classes/BehaviorTree/Tasks/BTTask_WaitBlackboardTime.h
+++ b/Engine/Source/Runtime/AIModule/Classes/BehaviorTree/Tasks/BTTask_WaitBlackboardTime.h
@@ -11,7 +11,9 @@
 UCLASS(hidecategories=Wait)
 class AIMODULE_API UBTTask_WaitBlackboardTime : public UBTTask_Wait
 {
-	GENERATED_UCLASS_BODY()
+	GENERATED_BODY()
+public:
+	UBTTask_WaitBlackboardTime(const FObjectInitializer& ObjectInitializer = FObjectInitializer::Get());
 
 	virtual void InitializeFromAsset(UBehaviorTree& Asset) override;
 	virtual EBTNodeResult::Type ExecuteTask(UBehaviorTreeComponent& OwnerComp, uint8* NodeMemory) override;
diff --git a/Engine/Source/Runtime/AIModule/Classes/Blueprint/AIAsyncTaskBlueprintProxy.h b/Engine/Source/Runtime/AIModule/Classes/Blueprint/AIAsyncTaskBlueprintProxy.h
--- a/Engine/Source/Runtime/AIModule/Classes/Blueprint/AIAsyncTaskBlueprintProxy.h
+++ b/Engine/Source/Runtime/AIModule/Classes/Blueprint/AIAsyncTaskBlueprintProxy.h
@@ -14,7 +14,9 @@ DECLARE_DYNAMIC_MULTICAST_DELEGATE_OneParam(FOAISimpleDelegate, EPathFollowingRe
 UCLASS(MinimalAPI)
 class UAIAsyncTaskBlueprintProxy : public UObject
 {
-	GENERATED_UCLASS_BODY()
+	GENERATED_BODY()
+public:
+	AIMODULE_API UAIAsyncTaskBlueprintProxy(const FObjectInitializer& ObjectInitializer = FObjectInitializer::Get());
 
 	UPROPERTY(BlueprintAssignable)
 	FOAISimpleDelegate	OnSuccess;
diff --git a/Engine/Source/Runtime/AIModule/Classes/Blueprint/AIBlueprintHelperLibrary.h b/Engine/Source/Runtime/AIModule/Classes/Blueprint/AIBlueprintHelperLibrary.h
--- a/Engine/Source/Runtime/AIModule/Classes/Blueprint/AIBlueprintHelperLibrary.h
+++ b/Engine/Source/Runtime/AIModule/Classes/Blueprint/AIBlueprintHelperLibrary.h
@@ -18,7 +18,9 @@ class UAIAsyncTaskBlueprintProxy;
 UCLASS()
 class AIMODULE_API UAIBlueprintHelperLibrary : public UBlueprintFunctionLibrary
 {
-	GENERATED_UCLASS_BODY()
+	GENERATED_BODY()
+public:
+	UAIBlueprintHelperLibrary(const FObjectInitializer& ObjectInitializer = FObjectInitializer::Get());
 
 	UFUNCTION(BlueprintCallable, meta=(WorldContext="WorldContextObject", BlueprintInternalUseOnly = "TRUE"))
 	static UAIAsyncTaskBlueprintProxy* CreateMoveToProxyObject(UObject* WorldContextObject, APawn* Pawn, FVector Destination, AActor* TargetActor = NULL, float AcceptanceRadius = 5.f, bool bStopOnOverlap = false);
diff --git a/Engine/Source/Runtime/AIModule/Classes/BrainComponent.h b/Engine/Source/Runtime/AIModule/Classes/BrainComponent.h
--- a/Engine/Source/Runtime/AIModule/Classes/BrainComponent.h
+++ b/Engine/Source/Runtime/AIModule/Classes/BrainComponent.h
@@ -105,7 +105,9 @@ private:
 UCLASS(abstract, BlueprintType)
 class AIMODULE_API UBrainComponent : public UActorComponent, public IAIResourceInterface
 {
-	GENERATED_UCLASS_BODY()
+	GENERATED_BODY()
+public:
+	UBrainComponent(const FObjectInitializer& ObjectInitializer = FObjectInitializer::Get());
 
 protected:
 	/** blackboard component */
@@ -209,4 +211,4 @@ FORCEINLINE UBlackboardComponent* UBrainComponent::GetBlackboardComponent()
 FORCEINLINE const UBlackboardComponent* UBrainComponent::GetBlackboardComponent() const
 {
 	return BlackboardComp;
-}
\ No newline at end of file
+}
diff --git a/Engine/Source/Runtime/AIModule/Classes/DataProviders/AIDataProvider.h b/Engine/Source/Runtime/AIModule/Classes/DataProviders/AIDataProvider.h
--- a/Engine/Source/Runtime/AIModule/Classes/DataProviders/AIDataProvider.h
+++ b/Engine/Source/Runtime/AIModule/Classes/DataProviders/AIDataProvider.h
@@ -20,7 +20,7 @@ class UAIDataProvider;
 USTRUCT()
 struct AIMODULE_API FAIDataProviderValue
 {
-	GENERATED_USTRUCT_BODY()
+	GENERATED_BODY()
 
 	/** cached uproperty of provider */
 	UPROPERTY(transient)
@@ -57,7 +57,7 @@ struct AIMODULE_API FAIDataProviderValue
 USTRUCT()
 struct AIMODULE_API FAIDataProviderTypedValue : public FAIDataProviderValue
 {
-	GENERATED_USTRUCT_BODY()
+	GENERATED_BODY()
 
 	/** type of value */
 	UPROPERTY()
@@ -70,7 +70,7 @@ struct AIMODULE_API FAIDataProviderTypedValue : public FAIDataProviderValue
 USTRUCT()
 struct AIMODULE_API FAIDataProviderStructValue : public FAIDataProviderValue
 {
-	GENERATED_USTRUCT_BODY()
+	GENERATED_BODY()
 
 	/** name of UStruct type */
 	FString StructName;
@@ -81,7 +81,7 @@ struct AIMODULE_API FAIDataProviderStructValue : public FAIDataProviderValue
 USTRUCT()
 struct AIMODULE_API FAIDataProviderIntValue : public FAIDataProviderTypedValue
 {
-	GENERATED_USTRUCT_BODY()
+	GENERATED_BODY()
 	FAIDataProviderIntValue();
 
 	UPROPERTY(EditAnywhere, Category = Value)
@@ -94,7 +94,7 @@ struct AIMODULE_API FAIDataProviderIntValue : public FAIDataProviderTypedValue
 USTRUCT()
 struct AIMODULE_API FAIDataProviderFloatValue : public FAIDataProviderTypedValue
 {
-	GENERATED_USTRUCT_BODY()
+	GENERATED_BODY()
 	FAIDataProviderFloatValue();
 
 	UPROPERTY(EditAnywhere, Category = Value)
@@ -107,7 +107,7 @@ struct AIMODULE_API FAIDataProviderFloatValue : public FAIDataProviderTypedValue
 USTRUCT()
 struct AIMODULE_API FAIDataProviderBoolValue : public FAIDataProviderTypedValue
 {
-	GENERATED_USTRUCT_BODY()
+	GENERATED_BODY()
 	FAIDataProviderBoolValue();
 
 	UPROPERTY(EditAnywhere, Category = Value)
@@ -120,7 +120,9 @@ struct AIMODULE_API FAIDataProviderBoolValue : public FAIDataProviderTypedValue
 UCLASS(EditInlineNew, Abstract, CollapseCategories, AutoExpandCategories=(Provider))
 class AIMODULE_API UAIDataProvider : public UObject
 {
-	GENERATED_UCLASS_BODY()
+	GENERATED_BODY()
+public:
+	UAIDataProvider(const FObjectInitializer& ObjectInitializer = FObjectInitializer::Get());
 
 	virtual void BindData(UObject* Owner, int32 RequestId);
 	virtual FString ToString(FName PropName) const;
diff --git a/Engine/Source/Runtime/AIModule/Classes/DataProviders/AIDataProvider_QueryParams.h b/Engine/Source/Runtime/AIModule/Classes/DataProviders/AIDataProvider_QueryParams.h
--- a/Engine/Source/Runtime/AIModule/Classes/DataProviders/AIDataProvider_QueryParams.h
+++ b/Engine/Source/Runtime/AIModule/Classes/DataProviders/AIDataProvider_QueryParams.h
@@ -16,7 +16,9 @@
 UCLASS(EditInlineNew, meta=(DisplayName="Query Params"))
 class AIMODULE_API UAIDataProvider_QueryParams : public UAIDataProvider
 {
-	GENERATED_UCLASS_BODY()
+	GENERATED_BODY()
+public:
+	UAIDataProvider_QueryParams(const FObjectInitializer& ObjectInitializer = FObjectInitializer::Get());
 
 	virtual void BindData(UObject* Owner, int32 RequestId);
 	virtual FString ToString(FName PropName) const override;
diff --git a/Engine/Source/Runtime/AIModule/Classes/EnvironmentQuery/Contexts/EnvQueryContext_BlueprintBase.h b/Engine/Source/Runtime/AIModule/Classes/EnvironmentQuery/Contexts/EnvQueryContext_BlueprintBase.h
--- a/Engine/Source/Runtime/AIModule/Classes/EnvironmentQuery/Contexts/EnvQueryContext_BlueprintBase.h
+++ b/Engine/Source/Runtime/AIModule/Classes/EnvironmentQuery/Contexts/EnvQueryContext_BlueprintBase.h
@@ -12,7 +12,9 @@ struct FEnvQueryContextData;
 UCLASS(MinimalAPI, Abstract, Blueprintable)
 class UEnvQueryContext_BlueprintBase : public UEnvQueryContext
 {
-	GENERATED_UCLASS_BODY()
+	GENERATED_BODY()
+public:
+	AIMODULE_API UEnvQueryContext_BlueprintBase(const FObjectInitializer& ObjectInitializer = FObjectInitializer::Get());
 
 	enum ECallMode
 	{
diff --git a/Engine/Source/Runtime/AIModule/Classes/EnvironmentQuery/Contexts/EnvQueryContext_Item.h b/Engine/Source/Runtime/AIModule/Classes/EnvironmentQuery/Contexts/EnvQueryContext_Item.h
--- a/Engine/Source/Runtime/AIModule/Classes/EnvironmentQuery/Contexts/EnvQueryContext_Item.h
+++ b/Engine/Source/Runtime/AIModule/Classes/EnvironmentQuery/Contexts/EnvQueryContext_Item.h
@@ -7,5 +7,7 @@
 UCLASS()
 class AIMODULE_API UEnvQueryContext_Item : public UEnvQueryContext
 {
-	GENERATED_UCLASS_BODY()
+	GENERATED_BODY()
+public:
+	UEnvQueryContext_Item(const FObjectInitializer& ObjectInitializer = FObjectInitializer::Get());
 };
diff --git a/Engine/Source/Runtime/AIModule/Classes/EnvironmentQuery/Contexts/EnvQueryContext_Querier.h b/Engine/Source/Runtime/AIModule/Classes/EnvironmentQuery/Contexts/EnvQueryContext_Querier.h
--- a/Engine/Source/Runtime/AIModule/Classes/EnvironmentQuery/Contexts/EnvQueryContext_Querier.h
+++ b/Engine/Source/Runtime/AIModule/Classes/EnvironmentQuery/Contexts/EnvQueryContext_Querier.h
@@ -10,7 +10,9 @@ struct FEnvQueryContextData;
 UCLASS(MinimalAPI)
 class UEnvQueryContext_Querier : public UEnvQueryContext
 {
-	GENERATED_UCLASS_BODY()
+	GENERATED_BODY()
+public:
+	AIMODULE_API UEnvQueryContext_Querier(const FObjectInitializer& ObjectInitializer = FObjectInitializer::Get());
  
 	virtual void ProvideContext(FEnvQueryInstance& QueryInstance, FEnvQueryContextData& ContextData) const override;
 };
diff --git a/Engine/Source/Runtime/AIModule/Classes/EnvironmentQuery/EQSQueryResultSourceInterface.h b/Engine/Source/Runtime/AIModule/Classes/EnvironmentQuery/EQSQueryResultSourceInterface.h
--- a/Engine/Source/Runtime/AIModule/Classes/EnvironmentQuery/EQSQueryResultSourceInterface.h
+++ b/Engine/Source/Runtime/AIModule/Classes/EnvironmentQuery/EQSQueryResultSourceInterface.h
@@ -7,12 +7,15 @@
 UINTERFACE(MinimalAPI, meta=(CannotImplementInterfaceInBlueprint))
 class UEQSQueryResultSourceInterface : public UInterface
 {
-	GENERATED_UINTERFACE_BODY()
+	GENERATED_BODY()
+public:
+	AIMODULE_API UEQSQueryResultSourceInterface(const FObjectInitializer& ObjectInitializer = FObjectInitializer::Get());
 };
 
 class IEQSQueryResultSourceInterface
 {
-	GENERATED_IINTERFACE_BODY()
+	GENERATED_BODY()
+public:
 
 	virtual const struct FEnvQueryResult* GetQueryResult() const { return NULL; }
 	virtual const struct FEnvQueryInstance* GetQueryInstance() const { return NULL; }
diff --git a/Engine/Source/Runtime/AIModule/Classes/EnvironmentQuery/EQSRenderingComponent.h b/Engine/Source/Runtime/AIModule/Classes/EnvironmentQuery/EQSRenderingComponent.h
--- a/Engine/Source/Runtime/AIModule/Classes/EnvironmentQuery/EQSRenderingComponent.h
+++ b/Engine/Source/Runtime/AIModule/Classes/EnvironmentQuery/EQSRenderingComponent.h
@@ -39,7 +39,9 @@ private:
 UCLASS(hidecategories=Object)
 class AIMODULE_API UEQSRenderingComponent : public UPrimitiveComponent
 {
-	GENERATED_UCLASS_BODY()
+	GENERATED_BODY()
+public:
+	UEQSRenderingComponent(const FObjectInitializer& ObjectInitializer = FObjectInitializer::Get());
 
 	FString DrawFlagName;
 	uint32 bDrawOnlyWhenSelected : 1;
diff --git a/Engine/Source/Runtime/AIModule/Classes/EnvironmentQuery/EQSTestingPawn.h b/Engine/Source/Runtime/AIModule/Classes/EnvironmentQuery/EQSTestingPawn.h
--- a/Engine/Source/Runtime/AIModule/Classes/EnvironmentQuery/EQSTestingPawn.h
+++ b/Engine/Source/Runtime/AIModule/Classes/EnvironmentQuery/EQSTestingPawn.h
@@ -24,7 +24,9 @@ enum class EEnvQueryHightlightMode : uint8
 UCLASS(abstract, hidecategories=(Advanced, Attachment, Collision, Mesh, Animation, Clothing, Physics, Rendering, Lighting, Activation, CharacterMovement, AgentPhysics, Avoidance, MovementComponent, Velocity, Shape, Camera, Input, Layers, SkeletalMesh, Optimization, Pawn, Replication, Actor))
 class AIMODULE_API AEQSTestingPawn : public ACharacter, public IEQSQueryResultSourceInterface
 {
-	GENERATED_UCLASS_BODY()
+	GENERATED_BODY()
+public:
+	AEQSTestingPawn(const FObjectInitializer& ObjectInitializer = FObjectInitializer::Get());
 	
 	UPROPERTY(Category=EQS, EditAnywhere)
 	UEnvQuery* QueryTemplate;
diff --git a/Engine/Source/Runtime/AIModule/Classes/EnvironmentQuery/EnvQuery.h b/Engine/Source/Runtime/AIModule/Classes/EnvironmentQuery/EnvQuery.h
--- a/Engine/Source/Runtime/AIModule/Classes/EnvironmentQuery/EnvQuery.h
+++ b/Engine/Source/Runtime/AIModule/Classes/EnvironmentQuery/EnvQuery.h
@@ -13,7 +13,9 @@ class UEdGraph;
 UCLASS()
 class AIMODULE_API UEnvQuery : public UObject
 {
-	GENERATED_UCLASS_BODY()
+	GENERATED_BODY()
+public:
+	UEnvQuery(const FObjectInitializer& ObjectInitializer = FObjectInitializer::Get());
 
 #if WITH_EDITORONLY_DATA
 	/** Graph for query */
diff --git a/Engine/Source/Runtime/AIModule/Classes/EnvironmentQuery/EnvQueryContext.h b/Engine/Source/Runtime/AIModule/Classes/EnvironmentQuery/EnvQueryContext.h
--- a/Engine/Source/Runtime/AIModule/Classes/EnvironmentQuery/EnvQueryContext.h
+++ b/Engine/Source/Runtime/AIModule/Classes/EnvironmentQuery/EnvQueryContext.h
@@ -9,7 +9,9 @@ struct FEnvQueryInstance;
 UCLASS(Abstract)
 class AIMODULE_API UEnvQueryContext : public UObject
 {
-	GENERATED_UCLASS_BODY()
+	GENERATED_BODY()
+public:
+	UEnvQueryContext(const FObjectInitializer& ObjectInitializer = FObjectInitializer::Get());
 
 	virtual void ProvideContext(FEnvQueryInstance& QueryInstance, FEnvQueryContextData& ContextData) const;
 };
diff --git a/Engine/Source/Runtime/AIModule/Classes/EnvironmentQuery/EnvQueryGenerator.h b/Engine/Source/Runtime/AIModule/Classes/EnvironmentQuery/EnvQueryGenerator.h
--- a/Engine/Source/Runtime/AIModule/Classes/EnvironmentQuery/EnvQueryGenerator.h
+++ b/Engine/Source/Runtime/AIModule/Classes/EnvironmentQuery/EnvQueryGenerator.h
@@ -19,7 +19,9 @@ namespace EnvQueryGeneratorVersion
 UCLASS(EditInlineNew, Abstract, meta = (Category = "Generators"))
 class AIMODULE_API UEnvQueryGenerator : public UEnvQueryNode
 {
-	GENERATED_UCLASS_BODY()
+	GENERATED_BODY()
+public:
+	UEnvQueryGenerator(const FObjectInitializer& ObjectInitializer = FObjectInitializer::Get());
 
 	UPROPERTY(EditDefaultsOnly, Category=Option)
 	FString OptionName;
diff --git a/Engine/Source/Runtime/AIModule/Classes/EnvironmentQuery/EnvQueryManager.h b/Engine/Source/Runtime/AIModule/Classes/EnvironmentQuery/EnvQueryManager.h
--- a/Engine/Source/Runtime/AIModule/Classes/EnvironmentQuery/EnvQueryManager.h
+++ b/Engine/Source/Runtime/AIModule/Classes/EnvironmentQuery/EnvQueryManager.h
@@ -18,7 +18,7 @@ class UEnvQueryTest;
 USTRUCT()
 struct AIMODULE_API FEnvQueryRequest
 {
-	GENERATED_USTRUCT_BODY()
+	GENERATED_BODY()
 
 	FEnvQueryRequest() : QueryTemplate(NULL), Owner(NULL), World(NULL) {}
 
@@ -74,7 +74,7 @@ protected:
 USTRUCT()
 struct FEnvQueryInstanceCache
 {
-	GENERATED_USTRUCT_BODY()
+	GENERATED_BODY()
 
 	/** query template, duplicated in manager's world */
 	UPROPERTY()
@@ -109,7 +109,9 @@ FORCEINLINE bool operator== (const FEQSDebugger::FEnvQueryInfo & Left, const FEQ
 UCLASS()
 class AIMODULE_API UEnvQueryManager : public UObject, public FTickableGameObject
 {
-	GENERATED_UCLASS_BODY()
+	GENERATED_BODY()
+public:
+	UEnvQueryManager(const FObjectInitializer& ObjectInitializer = FObjectInitializer::Get());
 
 	// We need to implement GetWorld() so that any EQS-related blueprints (such as blueprint contexts) can implement
 	// GetWorld() and so provide access to blueprint nodes using hidden WorldContextObject parameters.
diff --git a/Engine/Source/Runtime/AIModule/Classes/EnvironmentQuery/EnvQueryNode.h b/Engine/Source/Runtime/AIModule/Classes/EnvironmentQuery/EnvQueryNode.h
--- a/Engine/Source/Runtime/AIModule/Classes/EnvironmentQuery/EnvQueryNode.h
+++ b/Engine/Source/Runtime/AIModule/Classes/EnvironmentQuery/EnvQueryNode.h
@@ -6,7 +6,9 @@
 UCLASS(Abstract)
 class AIMODULE_API UEnvQueryNode : public UObject
 {
-	GENERATED_UCLASS_BODY()
+	GENERATED_BODY()
+public:
+	UEnvQueryNode(const FObjectInitializer& ObjectInitializer = FObjectInitializer::Get());
 
 	/** Versioning for updating deprecated properties */
 	UPROPERTY()
diff --git a/Engine/Source/Runtime/AIModule/Classes/EnvironmentQuery/EnvQueryOption.h b/Engine/Source/Runtime/AIModule/Classes/EnvironmentQuery/EnvQueryOption.h
--- a/Engine/Source/Runtime/AIModule/Classes/EnvironmentQuery/EnvQueryOption.h
+++ b/Engine/Source/Runtime/AIModule/Classes/EnvironmentQuery/EnvQueryOption.h
@@ -9,7 +9,9 @@ class UEnvQueryTest;
 UCLASS()
 class AIMODULE_API UEnvQueryOption : public UObject
 {
-	GENERATED_UCLASS_BODY()
+	GENERATED_BODY()
+public:
+	UEnvQueryOption(const FObjectInitializer& ObjectInitializer = FObjectInitializer::Get());
 
 	UPROPERTY()
 	UEnvQueryGenerator* Generator;
diff --git a/Engine/Source/Runtime/AIModule/Classes/EnvironmentQuery/EnvQueryTest.h b/Engine/Source/Runtime/AIModule/Classes/EnvironmentQuery/EnvQueryTest.h
--- a/Engine/Source/Runtime/AIModule/Classes/EnvironmentQuery/EnvQueryTest.h
+++ b/Engine/Source/Runtime/AIModule/Classes/EnvironmentQuery/EnvQueryTest.h
@@ -38,7 +38,9 @@ struct FEnvQueryTestScoringPreview
 UCLASS(Abstract)
 class AIMODULE_API UEnvQueryTest : public UEnvQueryNode
 {
-	GENERATED_UCLASS_BODY()
+	GENERATED_BODY()
+public:
+	UEnvQueryTest(const FObjectInitializer& ObjectInitializer = FObjectInitializer::Get());
 
 	/** Number of test as defined in data asset */
 	UPROPERTY()
diff --git a/Engine/Source/Runtime/AIModule/Classes/EnvironmentQuery/EnvQueryTypes.h b/Engine/Source/Runtime/AIModule/Classes/EnvironmentQuery/EnvQueryTypes.h
--- a/Engine/Source/Runtime/AIModule/Classes/EnvironmentQuery/EnvQueryTypes.h
+++ b/Engine/Source/Runtime/AIModule/Classes/EnvironmentQuery/EnvQueryTypes.h
@@ -190,7 +190,7 @@ namespace EEnvQueryTestClamping
 USTRUCT(BlueprintType)
 struct AIMODULE_API FEnvNamedValue
 {
-	GENERATED_USTRUCT_BODY();
+	GENERATED_BODY()
 
 	UPROPERTY(EditAnywhere, BlueprintReadWrite, Category=Param)
 	FName ParamName;
@@ -205,7 +205,7 @@ struct AIMODULE_API FEnvNamedValue
 USTRUCT()
 struct AIMODULE_API FEnvDirection
 {
-	GENERATED_USTRUCT_BODY()
+	GENERATED_BODY()
 
 	/** line A: start context */
 	UPROPERTY(EditDefaultsOnly, Category=Direction)
@@ -229,7 +229,7 @@ struct AIMODULE_API FEnvDirection
 USTRUCT()
 struct AIMODULE_API FEnvTraceData
 {
-	GENERATED_USTRUCT_BODY()
+	GENERATED_BODY()
 
 	enum EDescriptionMode
 	{
@@ -886,7 +886,7 @@ namespace FEQSHelpers
 USTRUCT()
 struct AIMODULE_API FEnvFloatParam_DEPRECATED
 {
-	GENERATED_USTRUCT_BODY();
+	GENERATED_BODY();
 
 	/** default value */
 	UPROPERTY(EditDefaultsOnly, Category = Param)
@@ -903,7 +903,7 @@ struct AIMODULE_API FEnvFloatParam_DEPRECATED
 USTRUCT()
 struct AIMODULE_API FEnvIntParam_DEPRECATED
 {
-	GENERATED_USTRUCT_BODY();
+	GENERATED_BODY();
 
 	/** default value */
 	UPROPERTY(EditDefaultsOnly, Category = Param)
@@ -920,7 +920,7 @@ struct AIMODULE_API FEnvIntParam_DEPRECATED
 USTRUCT()
 struct AIMODULE_API FEnvBoolParam_DEPRECATED
 {
-	GENERATED_USTRUCT_BODY();
+	GENERATED_BODY();
 
 	/** default value */
 	UPROPERTY(EditDefaultsOnly, Category = Param)
@@ -937,19 +937,19 @@ struct AIMODULE_API FEnvBoolParam_DEPRECATED
 USTRUCT()
 struct DEPRECATED(4.8, "FEnvFloatParam is deprecated in 4.8 and was replaced with FAIDataProviderFloatValue. Please use that type instead.") AIMODULE_API FEnvFloatParam : public FEnvFloatParam_DEPRECATED
 {
-	GENERATED_USTRUCT_BODY();
+	GENERATED_BODY();
 };
 
 USTRUCT()
 struct DEPRECATED(4.8, "FEnvIntParam is deprecated in 4.8 and was replaced with FAIDataProviderIntValue. Please use that type instead.") AIMODULE_API FEnvIntParam : public FEnvIntParam_DEPRECATED
 {
-	GENERATED_USTRUCT_BODY();
+	GENERATED_BODY();
 };
 
 USTRUCT()
 struct DEPRECATED(4.8, "FEnvBoolParam is deprecated in 4.8 and was replaced with FAIDataProviderBoolValue. Please use that type instead.") AIMODULE_API FEnvBoolParam : public FEnvBoolParam_DEPRECATED
 {
-	GENERATED_USTRUCT_BODY();
+	GENERATED_BODY();
 };
 
 // END DEPRECATED SUPPORT
diff --git a/Engine/Source/Runtime/AIModule/Classes/EnvironmentQuery/Generators/EnvQueryGenerator_ActorsOfClass.h b/Engine/Source/Runtime/AIModule/Classes/EnvironmentQuery/Generators/EnvQueryGenerator_ActorsOfClass.h
--- a/Engine/Source/Runtime/AIModule/Classes/EnvironmentQuery/Generators/EnvQueryGenerator_ActorsOfClass.h
+++ b/Engine/Source/Runtime/AIModule/Classes/EnvironmentQuery/Generators/EnvQueryGenerator_ActorsOfClass.h
@@ -10,7 +10,9 @@ class UEnvQueryContext;
 UCLASS(meta = (DisplayName = "Actors Of Class"))
 class UEnvQueryGenerator_ActorsOfClass : public UEnvQueryGenerator
 {
-	GENERATED_UCLASS_BODY()
+	GENERATED_BODY()
+public:
+	UEnvQueryGenerator_ActorsOfClass(const FObjectInitializer& ObjectInitializer = FObjectInitializer::Get());
 
 	/** max distance of path between point and context */
 	UPROPERTY(EditDefaultsOnly, Category=Generator)
diff --git a/Engine/Source/Runtime/AIModule/Classes/EnvironmentQuery/Generators/EnvQueryGenerator_BlueprintBase.h b/Engine/Source/Runtime/AIModule/Classes/EnvironmentQuery/Generators/EnvQueryGenerator_BlueprintBase.h
--- a/Engine/Source/Runtime/AIModule/Classes/EnvironmentQuery/Generators/EnvQueryGenerator_BlueprintBase.h
+++ b/Engine/Source/Runtime/AIModule/Classes/EnvironmentQuery/Generators/EnvQueryGenerator_BlueprintBase.h
@@ -12,7 +12,9 @@ class UEnvQueryItemType;
 UCLASS(Abstract, Blueprintable)
 class AIMODULE_API UEnvQueryGenerator_BlueprintBase : public UEnvQueryGenerator
 {
-	GENERATED_UCLASS_BODY()
+	GENERATED_BODY()
+public:
+	UEnvQueryGenerator_BlueprintBase(const FObjectInitializer& ObjectInitializer = FObjectInitializer::Get());
 
 	/** A short description of what test does, like "Generate pawn named Joe" */
 	UPROPERTY(EditAnywhere, Category = Generator)
diff --git a/Engine/Source/Runtime/AIModule/Classes/EnvironmentQuery/Generators/EnvQueryGenerator_Composite.h b/Engine/Source/Runtime/AIModule/Classes/EnvironmentQuery/Generators/EnvQueryGenerator_Composite.h
--- a/Engine/Source/Runtime/AIModule/Classes/EnvironmentQuery/Generators/EnvQueryGenerator_Composite.h
+++ b/Engine/Source/Runtime/AIModule/Classes/EnvironmentQuery/Generators/EnvQueryGenerator_Composite.h
@@ -15,7 +15,9 @@ struct FEnvQueryInstance;
 UCLASS(meta = (DisplayName = "Composite"))
 class AIMODULE_API UEnvQueryGenerator_Composite : public UEnvQueryGenerator
 {
-	GENERATED_UCLASS_BODY()
+	GENERATED_BODY()
+public:
+	UEnvQueryGenerator_Composite(const FObjectInitializer& ObjectInitializer = FObjectInitializer::Get());
 
 	UPROPERTY(EditDefaultsOnly, Instanced, Category = Generator)
 	TArray<UEnvQueryGenerator*> Generators;
diff --git a/Engine/Source/Runtime/AIModule/Classes/EnvironmentQuery/Generators/EnvQueryGenerator_Donut.h b/Engine/Source/Runtime/AIModule/Classes/EnvironmentQuery/Generators/EnvQueryGenerator_Donut.h
--- a/Engine/Source/Runtime/AIModule/Classes/EnvironmentQuery/Generators/EnvQueryGenerator_Donut.h
+++ b/Engine/Source/Runtime/AIModule/Classes/EnvironmentQuery/Generators/EnvQueryGenerator_Donut.h
@@ -9,7 +9,9 @@ struct FEnvQueryInstance;
 UCLASS(meta = (DisplayName = "Points: Donut"))
 class AIMODULE_API UEnvQueryGenerator_Donut : public UEnvQueryGenerator_ProjectedPoints
 {
-	GENERATED_UCLASS_BODY()
+	GENERATED_BODY()
+public:
+	UEnvQueryGenerator_Donut(const FObjectInitializer& ObjectInitializer = FObjectInitializer::Get());
 
 	/** min distance between point and context */
 	UPROPERTY(EditDefaultsOnly, Category = Generator)
diff --git a/Engine/Source/Runtime/AIModule/Classes/EnvironmentQuery/Generators/EnvQueryGenerator_OnCircle.h b/Engine/Source/Runtime/AIModule/Classes/EnvironmentQuery/Generators/EnvQueryGenerator_OnCircle.h
--- a/Engine/Source/Runtime/AIModule/Classes/EnvironmentQuery/Generators/EnvQueryGenerator_OnCircle.h
+++ b/Engine/Source/Runtime/AIModule/Classes/EnvironmentQuery/Generators/EnvQueryGenerator_OnCircle.h
@@ -9,7 +9,9 @@ struct FEnvQueryInstance;
 UCLASS(meta = (DisplayName = "Points: Circle"))
 class AIMODULE_API UEnvQueryGenerator_OnCircle : public UEnvQueryGenerator_ProjectedPoints
 {
-	GENERATED_UCLASS_BODY()
+	GENERATED_BODY()
+public:
+	UEnvQueryGenerator_OnCircle(const FObjectInitializer& ObjectInitializer = FObjectInitializer::Get());
 
 	/** max distance of path between point and context */
 	UPROPERTY(EditDefaultsOnly, Category=Generator)
diff --git a/Engine/Source/Runtime/AIModule/Classes/EnvironmentQuery/Generators/EnvQueryGenerator_PathingGrid.h b/Engine/Source/Runtime/AIModule/Classes/EnvironmentQuery/Generators/EnvQueryGenerator_PathingGrid.h
--- a/Engine/Source/Runtime/AIModule/Classes/EnvironmentQuery/Generators/EnvQueryGenerator_PathingGrid.h
+++ b/Engine/Source/Runtime/AIModule/Classes/EnvironmentQuery/Generators/EnvQueryGenerator_PathingGrid.h
@@ -14,7 +14,9 @@ class UNavigationQueryFilter;
 UCLASS(meta = (DisplayName = "Points: Pathing Grid"))
 class UEnvQueryGenerator_PathingGrid : public UEnvQueryGenerator_SimpleGrid
 {
-	GENERATED_UCLASS_BODY()
+	GENERATED_BODY()
+public:
+	UEnvQueryGenerator_PathingGrid(const FObjectInitializer& ObjectInitializer = FObjectInitializer::Get());
 
 	/** pathfinding direction */
 	UPROPERTY(EditDefaultsOnly, Category = Generator)
diff --git a/Engine/Source/Runtime/AIModule/Classes/EnvironmentQuery/Generators/EnvQueryGenerator_ProjectedPoints.h b/Engine/Source/Runtime/AIModule/Classes/EnvironmentQuery/Generators/EnvQueryGenerator_ProjectedPoints.h
--- a/Engine/Source/Runtime/AIModule/Classes/EnvironmentQuery/Generators/EnvQueryGenerator_ProjectedPoints.h
+++ b/Engine/Source/Runtime/AIModule/Classes/EnvironmentQuery/Generators/EnvQueryGenerator_ProjectedPoints.h
@@ -11,7 +11,9 @@ class ANavigationData;
 UCLASS(Abstract)
 class AIMODULE_API UEnvQueryGenerator_ProjectedPoints : public UEnvQueryGenerator
 {
-	GENERATED_UCLASS_BODY()
+	GENERATED_BODY()
+public:
+	UEnvQueryGenerator_ProjectedPoints(const FObjectInitializer& ObjectInitializer = FObjectInitializer::Get());
 
 	/** trace params */
 	UPROPERTY(EditDefaultsOnly, Category = Generator)
diff --git a/Engine/Source/Runtime/AIModule/Classes/EnvironmentQuery/Generators/EnvQueryGenerator_SimpleGrid.h b/Engine/Source/Runtime/AIModule/Classes/EnvironmentQuery/Generators/EnvQueryGenerator_SimpleGrid.h
--- a/Engine/Source/Runtime/AIModule/Classes/EnvironmentQuery/Generators/EnvQueryGenerator_SimpleGrid.h
+++ b/Engine/Source/Runtime/AIModule/Classes/EnvironmentQuery/Generators/EnvQueryGenerator_SimpleGrid.h
@@ -13,7 +13,9 @@ class UEnvQueryContext;
 UCLASS(meta = (DisplayName = "Points: Grid"))
 class UEnvQueryGenerator_SimpleGrid : public UEnvQueryGenerator_ProjectedPoints
 {
-	GENERATED_UCLASS_BODY()
+	GENERATED_BODY()
+public:
+	UEnvQueryGenerator_SimpleGrid(const FObjectInitializer& ObjectInitializer = FObjectInitializer::Get());
 
 	/** square's extent */
 	UPROPERTY(EditDefaultsOnly, Category=Generator)
diff --git a/Engine/Source/Runtime/AIModule/Classes/EnvironmentQuery/Items/EnvQueryItemType.h b/Engine/Source/Runtime/AIModule/Classes/EnvironmentQuery/Items/EnvQueryItemType.h
--- a/Engine/Source/Runtime/AIModule/Classes/EnvironmentQuery/Items/EnvQueryItemType.h
+++ b/Engine/Source/Runtime/AIModule/Classes/EnvironmentQuery/Items/EnvQueryItemType.h
@@ -9,7 +9,9 @@ struct FBlackboardKeySelector;
 UCLASS(Abstract)
 class AIMODULE_API UEnvQueryItemType : public UObject
 {
-	GENERATED_UCLASS_BODY()
+	GENERATED_BODY()
+public:
+	UEnvQueryItemType(const FObjectInitializer& ObjectInitializer = FObjectInitializer::Get());
 
 	/** get ValueSize */
 	FORCEINLINE uint16 GetValueSize() const { return ValueSize; }
diff --git a/Engine/Source/Runtime/AIModule/Classes/EnvironmentQuery/Items/EnvQueryItemType_Actor.h b/Engine/Source/Runtime/AIModule/Classes/EnvironmentQuery/Items/EnvQueryItemType_Actor.h
--- a/Engine/Source/Runtime/AIModule/Classes/EnvironmentQuery/Items/EnvQueryItemType_Actor.h
+++ b/Engine/Source/Runtime/AIModule/Classes/EnvironmentQuery/Items/EnvQueryItemType_Actor.h
@@ -10,7 +10,9 @@ struct FEnvQueryContextData;
 UCLASS()
 class AIMODULE_API UEnvQueryItemType_Actor : public UEnvQueryItemType_ActorBase
 {
-	GENERATED_UCLASS_BODY()
+	GENERATED_BODY()
+public:
+	UEnvQueryItemType_Actor(const FObjectInitializer& ObjectInitializer = FObjectInitializer::Get());
 
 	static AActor* GetValue(const uint8* RawData);
 	static void SetValue(uint8* RawData, const AActor* Value);
diff --git a/Engine/Source/Runtime/AIModule/Classes/EnvironmentQuery/Items/EnvQueryItemType_ActorBase.h b/Engine/Source/Runtime/AIModule/Classes/EnvironmentQuery/Items/EnvQueryItemType_ActorBase.h
--- a/Engine/Source/Runtime/AIModule/Classes/EnvironmentQuery/Items/EnvQueryItemType_ActorBase.h
+++ b/Engine/Source/Runtime/AIModule/Classes/EnvironmentQuery/Items/EnvQueryItemType_ActorBase.h
@@ -11,7 +11,9 @@ class UBlackboardComponent;
 UCLASS(Abstract)
 class AIMODULE_API UEnvQueryItemType_ActorBase : public UEnvQueryItemType_VectorBase
 {
-	GENERATED_UCLASS_BODY()
+	GENERATED_BODY()
+public:
+	UEnvQueryItemType_ActorBase(const FObjectInitializer& ObjectInitializer = FObjectInitializer::Get());
 
 	virtual AActor* GetActor(const uint8* RawData) const;
 
diff --git a/Engine/Source/Runtime/AIModule/Classes/EnvironmentQuery/Items/EnvQueryItemType_Direction.h b/Engine/Source/Runtime/AIModule/Classes/EnvironmentQuery/Items/EnvQueryItemType_Direction.h
--- a/Engine/Source/Runtime/AIModule/Classes/EnvironmentQuery/Items/EnvQueryItemType_Direction.h
+++ b/Engine/Source/Runtime/AIModule/Classes/EnvironmentQuery/Items/EnvQueryItemType_Direction.h
@@ -9,7 +9,9 @@ struct FEnvQueryContextData;
 UCLASS()
 class AIMODULE_API UEnvQueryItemType_Direction : public UEnvQueryItemType_VectorBase
 {
-	GENERATED_UCLASS_BODY()
+	GENERATED_BODY()
+public:
+	UEnvQueryItemType_Direction(const FObjectInitializer& ObjectInitializer = FObjectInitializer::Get());
 
 	static FVector GetValue(const uint8* RawData);
 	static void SetValue(uint8* RawData, const FVector& Value);
diff --git a/Engine/Source/Runtime/AIModule/Classes/EnvironmentQuery/Items/EnvQueryItemType_Point.h b/Engine/Source/Runtime/AIModule/Classes/EnvironmentQuery/Items/EnvQueryItemType_Point.h
--- a/Engine/Source/Runtime/AIModule/Classes/EnvironmentQuery/Items/EnvQueryItemType_Point.h
+++ b/Engine/Source/Runtime/AIModule/Classes/EnvironmentQuery/Items/EnvQueryItemType_Point.h
@@ -9,7 +9,9 @@ struct FEnvQueryContextData;
 UCLASS()
 class AIMODULE_API UEnvQueryItemType_Point : public UEnvQueryItemType_VectorBase
 {
-	GENERATED_UCLASS_BODY()
+	GENERATED_BODY()
+public:
+	UEnvQueryItemType_Point(const FObjectInitializer& ObjectInitializer = FObjectInitializer::Get());
 
 	static FVector GetValue(const uint8* RawData);
 	static void SetValue(uint8* RawData, const FVector& Value);
diff --git a/Engine/Source/Runtime/AIModule/Classes/EnvironmentQuery/Items/EnvQueryItemType_VectorBase.h b/Engine/Source/Runtime/AIModule/Classes/EnvironmentQuery/Items/EnvQueryItemType_VectorBase.h
--- a/Engine/Source/Runtime/AIModule/Classes/EnvironmentQuery/Items/EnvQueryItemType_VectorBase.h
+++ b/Engine/Source/Runtime/AIModule/Classes/EnvironmentQuery/Items/EnvQueryItemType_VectorBase.h
@@ -10,7 +10,9 @@ class UBlackboardComponent;
 UCLASS(Abstract)
 class AIMODULE_API UEnvQueryItemType_VectorBase : public UEnvQueryItemType
 {
-	GENERATED_UCLASS_BODY()
+	GENERATED_BODY()
+public:
+	UEnvQueryItemType_VectorBase(const FObjectInitializer& ObjectInitializer = FObjectInitializer::Get());
 
 	virtual FVector GetItemLocation(const uint8* RawData) const;
 	virtual FRotator GetItemRotation(const uint8* RawData) const;
diff --git a/Engine/Source/Runtime/AIModule/Classes/EnvironmentQuery/Tests/EnvQueryTest_Distance.h b/Engine/Source/Runtime/AIModule/Classes/EnvironmentQuery/Tests/EnvQueryTest_Distance.h
--- a/Engine/Source/Runtime/AIModule/Classes/EnvironmentQuery/Tests/EnvQueryTest_Distance.h
+++ b/Engine/Source/Runtime/AIModule/Classes/EnvironmentQuery/Tests/EnvQueryTest_Distance.h
@@ -20,7 +20,9 @@ namespace EEnvTestDistance
 UCLASS()
 class UEnvQueryTest_Distance : public UEnvQueryTest
 {
-	GENERATED_UCLASS_BODY()
+	GENERATED_BODY()
+public:
+	UEnvQueryTest_Distance(const FObjectInitializer& ObjectInitializer = FObjectInitializer::Get());
 
 	/** testing mode */
 	UPROPERTY(EditDefaultsOnly, Category=Distance)
diff --git a/Engine/Source/Runtime/AIModule/Classes/EnvironmentQuery/Tests/EnvQueryTest_Dot.h b/Engine/Source/Runtime/AIModule/Classes/EnvironmentQuery/Tests/EnvQueryTest_Dot.h
--- a/Engine/Source/Runtime/AIModule/Classes/EnvironmentQuery/Tests/EnvQueryTest_Dot.h
+++ b/Engine/Source/Runtime/AIModule/Classes/EnvironmentQuery/Tests/EnvQueryTest_Dot.h
@@ -18,7 +18,9 @@ enum class EEnvTestDot : uint8
 UCLASS(MinimalAPI)
 class UEnvQueryTest_Dot : public UEnvQueryTest
 {
-	GENERATED_UCLASS_BODY()
+	GENERATED_BODY()
+public:
+	AIMODULE_API UEnvQueryTest_Dot(const FObjectInitializer& ObjectInitializer = FObjectInitializer::Get());
 
 protected:
 	/** defines direction of first line used by test */
diff --git a/Engine/Source/Runtime/AIModule/Classes/EnvironmentQuery/Tests/EnvQueryTest_GameplayTags.h b/Engine/Source/Runtime/AIModule/Classes/EnvironmentQuery/Tests/EnvQueryTest_GameplayTags.h
--- a/Engine/Source/Runtime/AIModule/Classes/EnvironmentQuery/Tests/EnvQueryTest_GameplayTags.h
+++ b/Engine/Source/Runtime/AIModule/Classes/EnvironmentQuery/Tests/EnvQueryTest_GameplayTags.h
@@ -10,7 +10,9 @@ class IGameplayTagAssetInterface;
 UCLASS(MinimalAPI)
 class UEnvQueryTest_GameplayTags : public UEnvQueryTest
 {
-	GENERATED_UCLASS_BODY()
+	GENERATED_BODY()
+public:
+	AIMODULE_API UEnvQueryTest_GameplayTags(const FObjectInitializer& ObjectInitializer = FObjectInitializer::Get());
 
 protected:
 	virtual void RunTest(FEnvQueryInstance& QueryInstance) const override;
diff --git a/Engine/Source/Runtime/AIModule/Classes/EnvironmentQuery/Tests/EnvQueryTest_Pathfinding.h b/Engine/Source/Runtime/AIModule/Classes/EnvironmentQuery/Tests/EnvQueryTest_Pathfinding.h
--- a/Engine/Source/Runtime/AIModule/Classes/EnvironmentQuery/Tests/EnvQueryTest_Pathfinding.h
+++ b/Engine/Source/Runtime/AIModule/Classes/EnvironmentQuery/Tests/EnvQueryTest_Pathfinding.h
@@ -20,7 +20,9 @@ namespace EEnvTestPathfinding
 UCLASS()
 class UEnvQueryTest_Pathfinding : public UEnvQueryTest
 {
-	GENERATED_UCLASS_BODY()
+	GENERATED_BODY()
+public:
+	UEnvQueryTest_Pathfinding(const FObjectInitializer& ObjectInitializer = FObjectInitializer::Get());
 
 	/** testing mode */
 	UPROPERTY(EditDefaultsOnly, Category=Pathfinding)
diff --git a/Engine/Source/Runtime/AIModule/Classes/EnvironmentQuery/Tests/EnvQueryTest_Random.h b/Engine/Source/Runtime/AIModule/Classes/EnvironmentQuery/Tests/EnvQueryTest_Random.h
--- a/Engine/Source/Runtime/AIModule/Classes/EnvironmentQuery/Tests/EnvQueryTest_Random.h
+++ b/Engine/Source/Runtime/AIModule/Classes/EnvironmentQuery/Tests/EnvQueryTest_Random.h
@@ -7,7 +7,9 @@
 UCLASS(MinimalAPI, meta=(DeprecatedNode, DeprecationMessage = "This class is now deprecated, please use RunMode supporting random results instead."))
 class UEnvQueryTest_Random : public UEnvQueryTest
 {
-	GENERATED_UCLASS_BODY()
+	GENERATED_BODY()
+public:
+	AIMODULE_API UEnvQueryTest_Random(const FObjectInitializer& ObjectInitializer = FObjectInitializer::Get());
 
 protected:
 	virtual void RunTest(FEnvQueryInstance& QueryInstance) const override;
diff --git a/Engine/Source/Runtime/AIModule/Classes/EnvironmentQuery/Tests/EnvQueryTest_Trace.h b/Engine/Source/Runtime/AIModule/Classes/EnvironmentQuery/Tests/EnvQueryTest_Trace.h
--- a/Engine/Source/Runtime/AIModule/Classes/EnvironmentQuery/Tests/EnvQueryTest_Trace.h
+++ b/Engine/Source/Runtime/AIModule/Classes/EnvironmentQuery/Tests/EnvQueryTest_Trace.h
@@ -9,7 +9,9 @@ class UEnvQueryContext;
 UCLASS(MinimalAPI)
 class UEnvQueryTest_Trace : public UEnvQueryTest
 {
-	GENERATED_UCLASS_BODY()
+	GENERATED_BODY()
+public:
+	AIMODULE_API UEnvQueryTest_Trace(const FObjectInitializer& ObjectInitializer = FObjectInitializer::Get());
 
 	/** trace data */
 	UPROPERTY(EditDefaultsOnly, Category=Trace)
diff --git a/Engine/Source/Runtime/AIModule/Classes/GenericTeamAgentInterface.h b/Engine/Source/Runtime/AIModule/Classes/GenericTeamAgentInterface.h
--- a/Engine/Source/Runtime/AIModule/Classes/GenericTeamAgentInterface.h
+++ b/Engine/Source/Runtime/AIModule/Classes/GenericTeamAgentInterface.h
@@ -18,7 +18,7 @@ namespace ETeamAttitude
 USTRUCT()
 struct AIMODULE_API FGenericTeamId
 {
-	GENERATED_USTRUCT_BODY()
+	GENERATED_BODY()
 
 private:
 	enum EPredefinedId
@@ -49,12 +49,15 @@ public:
 UINTERFACE()
 class AIMODULE_API UGenericTeamAgentInterface : public UInterface
 {
-	GENERATED_UINTERFACE_BODY()
+	GENERATED_BODY()
+public:
+	UGenericTeamAgentInterface(const FObjectInitializer& ObjectInitializer = FObjectInitializer::Get());
 };
 
 class AIMODULE_API IGenericTeamAgentInterface
 {
-	GENERATED_IINTERFACE_BODY()
+	GENERATED_BODY()
+public:
 
 	/** Assigns Team Agent to given TeamID */
 	virtual void SetGenericTeamId(const FGenericTeamId& TemaID) {}
diff --git a/Engine/Source/Runtime/AIModule/Classes/HotSpots/AIHotSpotManager.h b/Engine/Source/Runtime/AIModule/Classes/HotSpots/AIHotSpotManager.h
--- a/Engine/Source/Runtime/AIModule/Classes/HotSpots/AIHotSpotManager.h
+++ b/Engine/Source/Runtime/AIModule/Classes/HotSpots/AIHotSpotManager.h
@@ -6,5 +6,7 @@
 UCLASS()
 class AIMODULE_API UAIHotSpotManager : public UObject
 {
-	GENERATED_UCLASS_BODY()
-};
\ No newline at end of file
+	GENERATED_BODY()
+public:
+	UAIHotSpotManager(const FObjectInitializer& ObjectInitializer = FObjectInitializer::Get());
+};
diff --git a/Engine/Source/Runtime/AIModule/Classes/Navigation/CrowdAgentInterface.h b/Engine/Source/Runtime/AIModule/Classes/Navigation/CrowdAgentInterface.h
--- a/Engine/Source/Runtime/AIModule/Classes/Navigation/CrowdAgentInterface.h
+++ b/Engine/Source/Runtime/AIModule/Classes/Navigation/CrowdAgentInterface.h
@@ -7,12 +7,15 @@
 UINTERFACE(MinimalAPI, meta=(CannotImplementInterfaceInBlueprint))
 class UCrowdAgentInterface : public UInterface
 {
-	GENERATED_UINTERFACE_BODY()
+	GENERATED_BODY()
+public:
+	AIMODULE_API UCrowdAgentInterface(const FObjectInitializer& ObjectInitializer = FObjectInitializer::Get());
 };
 
 class ICrowdAgentInterface
 {
-	GENERATED_IINTERFACE_BODY()
+	GENERATED_BODY()
+public:
 
 	/** @return current location of crowd agent */
 	virtual FVector GetCrowdAgentLocation() const { return FAISystem::InvalidLocation; }
diff --git a/Engine/Source/Runtime/AIModule/Classes/Navigation/CrowdFollowingComponent.h b/Engine/Source/Runtime/AIModule/Classes/Navigation/CrowdFollowingComponent.h
--- a/Engine/Source/Runtime/AIModule/Classes/Navigation/CrowdFollowingComponent.h
+++ b/Engine/Source/Runtime/AIModule/Classes/Navigation/CrowdFollowingComponent.h
@@ -24,7 +24,9 @@ namespace ECrowdAvoidanceQuality
 UCLASS(BlueprintType)
 class AIMODULE_API UCrowdFollowingComponent : public UPathFollowingComponent, public ICrowdAgentInterface
 {
-	GENERATED_UCLASS_BODY()
+	GENERATED_BODY()
+public:
+	UCrowdFollowingComponent(const FObjectInitializer& ObjectInitializer = FObjectInitializer::Get());
 
 	UPROPERTY()
 	FVector CrowdAgentMoveDirection;
diff --git a/Engine/Source/Runtime/AIModule/Classes/Navigation/CrowdManager.h b/Engine/Source/Runtime/AIModule/Classes/Navigation/CrowdManager.h
--- a/Engine/Source/Runtime/AIModule/Classes/Navigation/CrowdManager.h
+++ b/Engine/Source/Runtime/AIModule/Classes/Navigation/CrowdManager.h
@@ -52,7 +52,7 @@ class dtObstacleAvoidanceDebugData;
 USTRUCT()
 struct AIMODULE_API FCrowdAvoidanceConfig
 {
-	GENERATED_USTRUCT_BODY()
+	GENERATED_BODY()
 
 	UPROPERTY(EditAnywhere, Category=Crowd)
 	float VelocityBias;
@@ -98,7 +98,7 @@ struct AIMODULE_API FCrowdAvoidanceConfig
 USTRUCT()
 struct AIMODULE_API FCrowdAvoidanceSamplingPattern
 {
-	GENERATED_USTRUCT_BODY()
+	GENERATED_BODY()
 
 	UPROPERTY(EditAnywhere, Category=Crowd)
 	TArray<float> Angles;
@@ -153,7 +153,9 @@ struct FCrowdTickHelper : FTickableGameObject
 UCLASS(config = Engine, defaultconfig)
 class AIMODULE_API UCrowdManager : public UObject
 {
-	GENERATED_UCLASS_BODY()
+	GENERATED_BODY()
+public:
+	UCrowdManager(const FObjectInitializer& ObjectInitializer = FObjectInitializer::Get());
 
 	virtual void Tick(float DeltaTime);
 	virtual void BeginDestroy() override;
diff --git a/Engine/Source/Runtime/AIModule/Classes/Navigation/PathFollowingComponent.h b/Engine/Source/Runtime/AIModule/Classes/Navigation/PathFollowingComponent.h
--- a/Engine/Source/Runtime/AIModule/Classes/Navigation/PathFollowingComponent.h
+++ b/Engine/Source/Runtime/AIModule/Classes/Navigation/PathFollowingComponent.h
@@ -90,7 +90,9 @@ namespace EPathFollowingMessage
 UCLASS(config=Engine)
 class AIMODULE_API UPathFollowingComponent : public UActorComponent, public IAIResourceInterface
 {
-	GENERATED_UCLASS_BODY()
+	GENERATED_BODY()
+public:
+	UPathFollowingComponent(const FObjectInitializer& ObjectInitializer = FObjectInitializer::Get());
 
 	DECLARE_DELEGATE_TwoParams(FPostProcessMoveSignature, UPathFollowingComponent* /*comp*/, FVector& /*velocity*/);
 	DECLARE_DELEGATE_OneParam(FRequestCompletedSignature, EPathFollowingResult::Type /*Result*/);
diff --git a/Engine/Source/Runtime/AIModule/Classes/Perception/AIPerceptionComponent.h b/Engine/Source/Runtime/AIModule/Classes/Perception/AIPerceptionComponent.h
--- a/Engine/Source/Runtime/AIModule/Classes/Perception/AIPerceptionComponent.h
+++ b/Engine/Source/Runtime/AIModule/Classes/Perception/AIPerceptionComponent.h
@@ -84,7 +84,7 @@ struct AIMODULE_API FActorPerceptionInfo
 USTRUCT(BlueprintType, meta = (DisplayName = "Sensed Actor's Data"))
 struct FActorPerceptionBlueprintInfo
 {
-	GENERATED_USTRUCT_BODY()
+	GENERATED_BODY()
 
 	UPROPERTY(BlueprintReadWrite, Category = "AI|Perception")
 	AActor* Target;
@@ -107,7 +107,9 @@ struct FActorPerceptionBlueprintInfo
 UCLASS(ClassGroup=AI, HideCategories=(Activation, Collision), meta=(BlueprintSpawnableComponent), config=Game)
 class AIMODULE_API UAIPerceptionComponent : public UActorComponent
 {
-	GENERATED_UCLASS_BODY()
+	GENERATED_BODY()
+public:
+	UAIPerceptionComponent(const FObjectInitializer& ObjectInitializer = FObjectInitializer::Get());
 	
 	static const int32 InitialStimuliToProcessArraySize;
 
diff --git a/Engine/Source/Runtime/AIModule/Classes/Perception/AIPerceptionListenerInterface.h b/Engine/Source/Runtime/AIModule/Classes/Perception/AIPerceptionListenerInterface.h
--- a/Engine/Source/Runtime/AIModule/Classes/Perception/AIPerceptionListenerInterface.h
+++ b/Engine/Source/Runtime/AIModule/Classes/Perception/AIPerceptionListenerInterface.h
@@ -9,12 +9,15 @@ class UAIPerceptionComponent;
 UINTERFACE()
 class AIMODULE_API UAIPerceptionListenerInterface : public UInterface
 {
-	GENERATED_UINTERFACE_BODY()
+	GENERATED_BODY()
+public:
+	UAIPerceptionListenerInterface(const FObjectInitializer& ObjectInitializer = FObjectInitializer::Get());
 };
 
 class AIMODULE_API IAIPerceptionListenerInterface
 {
-	GENERATED_IINTERFACE_BODY()
+	GENERATED_BODY()
+public:
 
 	virtual UAIPerceptionComponent* GetPerceptionComponent() { return NULL; }
 };
diff --git a/Engine/Source/Runtime/AIModule/Classes/Perception/AIPerceptionTypes.h b/Engine/Source/Runtime/AIModule/Classes/Perception/AIPerceptionTypes.h
--- a/Engine/Source/Runtime/AIModule/Classes/Perception/AIPerceptionTypes.h
+++ b/Engine/Source/Runtime/AIModule/Classes/Perception/AIPerceptionTypes.h
@@ -61,7 +61,7 @@ struct FPerceptionChannelFilter
 USTRUCT(BlueprintType)
 struct AIMODULE_API FAIStimulus
 {
-	GENERATED_USTRUCT_BODY()
+	GENERATED_BODY()
 
 	static const float NeverHappenedAge;
 
@@ -130,7 +130,7 @@ public:
 USTRUCT()
 struct AIMODULE_API FAISenseAffiliationFilter
 {
-	GENERATED_USTRUCT_BODY()
+	GENERATED_BODY()
 
 	UPROPERTY(EditDefaultsOnly, BlueprintReadOnly, Category = "Sense")
 	uint32 bDetectEnemies : 1;
diff --git a/Engine/Source/Runtime/AIModule/Classes/Perception/AISense.h b/Engine/Source/Runtime/AIModule/Classes/Perception/AISense.h
--- a/Engine/Source/Runtime/AIModule/Classes/Perception/AISense.h
+++ b/Engine/Source/Runtime/AIModule/Classes/Perception/AISense.h
@@ -13,7 +13,9 @@ DECLARE_DELEGATE_OneParam(FOnPerceptionListenerUpdateDelegate, const FPerception
 UCLASS(ClassGroup = AI, abstract, config = Engine)
 class AIMODULE_API UAISense : public UObject
 {
-	GENERATED_UCLASS_BODY()
+	GENERATED_BODY()
+public:
+	UAISense(const FObjectInitializer& ObjectInitializer = FObjectInitializer::Get());
 
 	static const float SuspendNextUpdate;
 
diff --git a/Engine/Source/Runtime/AIModule/Classes/Perception/AISense_Damage.h b/Engine/Source/Runtime/AIModule/Classes/Perception/AISense_Damage.h
--- a/Engine/Source/Runtime/AIModule/Classes/Perception/AISense_Damage.h
+++ b/Engine/Source/Runtime/AIModule/Classes/Perception/AISense_Damage.h
@@ -9,7 +9,7 @@
 USTRUCT()
 struct AIMODULE_API FAIDamageEvent
 {	
-	GENERATED_USTRUCT_BODY()
+	GENERATED_BODY()
 
 	typedef class UAISense_Damage FSenseClass;
 
@@ -49,7 +49,9 @@ struct AIMODULE_API FAIDamageEvent
 UCLASS(ClassGroup=AI)
 class AIMODULE_API UAISense_Damage : public UAISense
 {
-	GENERATED_UCLASS_BODY()
+	GENERATED_BODY()
+public:
+	UAISense_Damage(const FObjectInitializer& ObjectInitializer = FObjectInitializer::Get());
 
 	UPROPERTY()
 	TArray<FAIDamageEvent> RegisteredEvents;
diff --git a/Engine/Source/Runtime/AIModule/Classes/Perception/AISense_Hearing.h b/Engine/Source/Runtime/AIModule/Classes/Perception/AISense_Hearing.h
--- a/Engine/Source/Runtime/AIModule/Classes/Perception/AISense_Hearing.h
+++ b/Engine/Source/Runtime/AIModule/Classes/Perception/AISense_Hearing.h
@@ -11,7 +11,7 @@ class UAISenseConfig_Hearing;
 USTRUCT()
 struct AIMODULE_API FAINoiseEvent
 {	
-	GENERATED_USTRUCT_BODY()
+	GENERATED_BODY()
 
 	typedef class UAISense_Hearing FSenseClass;
 
@@ -39,7 +39,9 @@ struct AIMODULE_API FAINoiseEvent
 UCLASS(ClassGroup=AI, Config=Game)
 class AIMODULE_API UAISense_Hearing : public UAISense
 {
-	GENERATED_UCLASS_BODY()
+	GENERATED_BODY()
+public:
+	UAISense_Hearing(const FObjectInitializer& ObjectInitializer = FObjectInitializer::Get());
 		
 protected:
 	UPROPERTY()
diff --git a/Engine/Source/Runtime/AIModule/Classes/Perception/AISense_Prediction.h b/Engine/Source/Runtime/AIModule/Classes/Perception/AISense_Prediction.h
--- a/Engine/Source/Runtime/AIModule/Classes/Perception/AISense_Prediction.h
+++ b/Engine/Source/Runtime/AIModule/Classes/Perception/AISense_Prediction.h
@@ -11,7 +11,7 @@ class UAISense_Prediction;
 USTRUCT()
 struct AIMODULE_API FAIPredictionEvent
 {	
-	GENERATED_USTRUCT_BODY()
+	GENERATED_BODY()
 
 	typedef UAISense_Prediction FSenseClass;
 	
@@ -34,7 +34,9 @@ struct AIMODULE_API FAIPredictionEvent
 UCLASS(ClassGroup=AI)
 class AIMODULE_API UAISense_Prediction : public UAISense
 {
-	GENERATED_UCLASS_BODY()
+	GENERATED_BODY()
+public:
+	UAISense_Prediction(const FObjectInitializer& ObjectInitializer = FObjectInitializer::Get());
 
 	UPROPERTY()
 	TArray<FAIPredictionEvent> RegisteredEvents;
diff --git a/Engine/Source/Runtime/AIModule/Classes/Perception/AISense_Sight.h b/Engine/Source/Runtime/AIModule/Classes/Perception/AISense_Sight.h
--- a/Engine/Source/Runtime/AIModule/Classes/Perception/AISense_Sight.h
+++ b/Engine/Source/Runtime/AIModule/Classes/Perception/AISense_Sight.h
@@ -23,7 +23,7 @@ namespace ESightPerceptionEventName
 USTRUCT()
 struct AIMODULE_API FAISightEvent
 {
-	GENERATED_USTRUCT_BODY()
+	GENERATED_BODY()
 
 	typedef UAISense_Sight FSenseClass;
 
@@ -101,7 +101,9 @@ struct FAISightQuery
 UCLASS(ClassGroup=AI, config=Game)
 class AIMODULE_API UAISense_Sight : public UAISense
 {
-	GENERATED_UCLASS_BODY()
+	GENERATED_BODY()
+public:
+	UAISense_Sight(const FObjectInitializer& ObjectInitializer = FObjectInitializer::Get());
 
 public:
 	struct FDigestedSightProperties
diff --git a/Engine/Source/Runtime/AIModule/Classes/Perception/AISense_Team.h b/Engine/Source/Runtime/AIModule/Classes/Perception/AISense_Team.h
--- a/Engine/Source/Runtime/AIModule/Classes/Perception/AISense_Team.h
+++ b/Engine/Source/Runtime/AIModule/Classes/Perception/AISense_Team.h
@@ -11,7 +11,7 @@ class UAISense_Team;
 USTRUCT()
 struct AIMODULE_API FAITeamStimulusEvent
 {	
-	GENERATED_USTRUCT_BODY()
+	GENERATED_BODY()
 
 	typedef UAISense_Team FSenseClass;
 
@@ -47,7 +47,9 @@ public:
 UCLASS(ClassGroup=AI)
 class AIMODULE_API UAISense_Team : public UAISense
 {
-	GENERATED_UCLASS_BODY()
+	GENERATED_BODY()
+public:
+	UAISense_Team(const FObjectInitializer& ObjectInitializer = FObjectInitializer::Get());
 
 	UPROPERTY()
 	TArray<FAITeamStimulusEvent> RegisteredEvents;
diff --git a/Engine/Source/Runtime/AIModule/Classes/Perception/AISense_Touch.h b/Engine/Source/Runtime/AIModule/Classes/Perception/AISense_Touch.h
--- a/Engine/Source/Runtime/AIModule/Classes/Perception/AISense_Touch.h
+++ b/Engine/Source/Runtime/AIModule/Classes/Perception/AISense_Touch.h
@@ -11,7 +11,7 @@ class UAISense_Touch;
 USTRUCT()
 struct AIMODULE_API FAITouchEvent
 {	
-	GENERATED_USTRUCT_BODY()
+	GENERATED_BODY()
 
 	typedef UAISense_Touch FSenseClass;
 
@@ -34,7 +34,9 @@ struct AIMODULE_API FAITouchEvent
 UCLASS(ClassGroup=AI)
 class AIMODULE_API UAISense_Touch : public UAISense
 {
-	GENERATED_UCLASS_BODY()
+	GENERATED_BODY()
+public:
+	UAISense_Touch(const FObjectInitializer& ObjectInitializer = FObjectInitializer::Get());
 
 	UPROPERTY()
 	TArray<FAITouchEvent> RegisteredEvents;
diff --git a/Engine/Source/Runtime/AIModule/Classes/Perception/AISightTargetInterface.h b/Engine/Source/Runtime/AIModule/Classes/Perception/AISightTargetInterface.h
--- a/Engine/Source/Runtime/AIModule/Classes/Perception/AISightTargetInterface.h
+++ b/Engine/Source/Runtime/AIModule/Classes/Perception/AISightTargetInterface.h
@@ -7,12 +7,15 @@
 UINTERFACE()
 class AIMODULE_API UAISightTargetInterface : public UInterface
 {
-	GENERATED_UINTERFACE_BODY()
+	GENERATED_BODY()
+public:
+	UAISightTargetInterface(const FObjectInitializer& ObjectInitializer = FObjectInitializer::Get());
 };
 
 class AIMODULE_API IAISightTargetInterface
 {
-	GENERATED_IINTERFACE_BODY()
+	GENERATED_BODY()
+public:
 
 	/**	Implementation should check whether from given ObserverLocation
 	 *	implementer can be seen. If so OutSeenLocation should contain
diff --git a/Engine/Source/Runtime/AIModule/Classes/Perception/PawnSensingComponent.h b/Engine/Source/Runtime/AIModule/Classes/Perception/PawnSensingComponent.h
--- a/Engine/Source/Runtime/AIModule/Classes/Perception/PawnSensingComponent.h
+++ b/Engine/Source/Runtime/AIModule/Classes/Perception/PawnSensingComponent.h
@@ -16,7 +16,9 @@ class APawn;
 UCLASS(ClassGroup=AI, HideCategories=(Activation, "Components|Activation", Collision), meta=(BlueprintSpawnableComponent))
 class AIMODULE_API UPawnSensingComponent : public UActorComponent
 {
-	GENERATED_UCLASS_BODY()
+	GENERATED_BODY()
+public:
+	UPawnSensingComponent(const FObjectInitializer& ObjectInitializer = FObjectInitializer::Get());
 
 	DECLARE_DYNAMIC_MULTICAST_DELEGATE_OneParam( FSeePawnDelegate, APawn*, Pawn );
 	DECLARE_DYNAMIC_MULTICAST_DELEGATE_ThreeParams( FHearNoiseDelegate, APawn*, Instigator, const FVector&, Location, float, Volume);
diff --git a/Engine/Source/Runtime/Android/AndroidRuntimeSettings/Classes/AndroidRuntimeSettings.h b/Engine/Source/Runtime/Android/AndroidRuntimeSettings/Classes/AndroidRuntimeSettings.h
--- a/Engine/Source/Runtime/Android/AndroidRuntimeSettings/Classes/AndroidRuntimeSettings.h
+++ b/Engine/Source/Runtime/Android/AndroidRuntimeSettings/Classes/AndroidRuntimeSettings.h
@@ -56,7 +56,7 @@ namespace EAndroidDepthBufferPreference
 USTRUCT()
 struct FGooglePlayAchievementMapping
 {
-	GENERATED_USTRUCT_BODY()
+	GENERATED_BODY()
 
 	// The game-specific achievement name (the one passed in to WriteAchievement calls)
 	UPROPERTY(EditAnywhere, Category = GooglePlayServices)
@@ -73,7 +73,7 @@ struct FGooglePlayAchievementMapping
 USTRUCT()
 struct FGooglePlayLeaderboardMapping
 {
-	GENERATED_USTRUCT_BODY()
+	GENERATED_BODY()
 
 	// The game-specific leaderboard name (the one passed in to WriteLeaderboards calls)
 	UPROPERTY(EditAnywhere, Category = GooglePlayServices)
@@ -101,8 +101,10 @@ namespace EAndroidAudio
 UCLASS(config=Engine, defaultconfig)
 class ANDROIDRUNTIMESETTINGS_API UAndroidRuntimeSettings : public UObject
 {
+	GENERATED_BODY()
+public:
+	UAndroidRuntimeSettings(const FObjectInitializer& ObjectInitializer = FObjectInitializer::Get());
 public:
-	GENERATED_UCLASS_BODY()
 
 	// The official name of the product (same as the name you use on the Play Store web site)
 	UPROPERTY(GlobalConfig, EditAnywhere, Category = APKPackaging, Meta = (DisplayName = "Android Package Name ('com.Company.Project', [PROJECT] is replaced with project name)"))
diff --git a/Engine/Source/Runtime/AutomationMessages/Classes/AutomationWorkerMessages.h b/Engine/Source/Runtime/AutomationMessages/Classes/AutomationWorkerMessages.h
--- a/Engine/Source/Runtime/AutomationMessages/Classes/AutomationWorkerMessages.h
+++ b/Engine/Source/Runtime/AutomationMessages/Classes/AutomationWorkerMessages.h
@@ -14,7 +14,7 @@
 USTRUCT()
 struct FAutomationWorkerFindWorkers
 {
-	GENERATED_USTRUCT_BODY()
+	GENERATED_BODY()
 
 	/** Holds the change list number to find workers for. */
 	UPROPERTY()
@@ -61,7 +61,7 @@ struct TStructOpsTypeTraits<FAutomationWorkerFindWorkers> : public TStructOpsTyp
 USTRUCT()
 struct FAutomationWorkerFindWorkersResponse
 {
-	GENERATED_USTRUCT_BODY()
+	GENERATED_BODY()
 
 	/** Holds the name of the device that the worker is running on. */
 	UPROPERTY()
@@ -125,7 +125,7 @@ struct TStructOpsTypeTraits<FAutomationWorkerFindWorkersResponse> : public TStru
 USTRUCT()
 struct FAutomationWorkerWorkerOffline
 {
-	GENERATED_USTRUCT_BODY()
+	GENERATED_BODY()
 };
 
 template<>
@@ -143,7 +143,7 @@ struct TStructOpsTypeTraits<FAutomationWorkerWorkerOffline> : public TStructOpsT
 USTRUCT()
 struct FAutomationWorkerPing
 {
-	GENERATED_USTRUCT_BODY()
+	GENERATED_BODY()
 };
 
 template<>
@@ -161,7 +161,7 @@ struct TStructOpsTypeTraits<FAutomationWorkerPing> : public TStructOpsTypeTraits
 USTRUCT()
 struct FAutomationWorkerResetTests
 {
-	GENERATED_USTRUCT_BODY()
+	GENERATED_BODY()
 };
 
 template<>
@@ -179,7 +179,7 @@ struct TStructOpsTypeTraits<FAutomationWorkerResetTests> : public TStructOpsType
 USTRUCT()
 struct FAutomationWorkerPong
 {
-	GENERATED_USTRUCT_BODY()
+	GENERATED_BODY()
 };
 
 template<>
@@ -198,7 +198,7 @@ struct TStructOpsTypeTraits<FAutomationWorkerPong> : public TStructOpsTypeTraits
 USTRUCT()
 struct FAutomationWorkerRequestTests
 {
-	GENERATED_USTRUCT_BODY()
+	GENERATED_BODY()
 
 	/** Holds a flag indicating whether the developer directory should be included. */
 	UPROPERTY()
@@ -238,7 +238,7 @@ struct TStructOpsTypeTraits<FAutomationWorkerRequestTests> : public TStructOpsTy
 USTRUCT()
 struct FAutomationWorkerRequestTestsReply
 {
-	GENERATED_USTRUCT_BODY()
+	GENERATED_BODY()
 
 	/** Holds the test information serialized into a string. */
 	UPROPERTY()
@@ -278,7 +278,7 @@ struct TStructOpsTypeTraits<FAutomationWorkerRequestTestsReply> : public TStruct
 USTRUCT()
 struct FAutomationWorkerRunTests
 {
-	GENERATED_USTRUCT_BODY()
+	GENERATED_BODY()
 
 	/** */
 	UPROPERTY()
@@ -333,7 +333,7 @@ struct TStructOpsTypeTraits<FAutomationWorkerRunTests> : public TStructOpsTypeTr
 USTRUCT()
 struct FAutomationWorkerRunTestsReply
 {
-	GENERATED_USTRUCT_BODY()
+	GENERATED_BODY()
 
 	/** */
 	UPROPERTY()
@@ -379,7 +379,7 @@ struct TStructOpsTypeTraits<FAutomationWorkerRunTestsReply> : public TStructOpsT
 USTRUCT()
 struct FAutomationWorkerRequestNextNetworkCommand
 {
-	GENERATED_USTRUCT_BODY()
+	GENERATED_BODY()
 
 	/** */
 	UPROPERTY()
@@ -413,7 +413,7 @@ struct TStructOpsTypeTraits<FAutomationWorkerRequestNextNetworkCommand> : public
 USTRUCT()
 struct FAutomationWorkerNextNetworkCommandReply
 {
-	GENERATED_USTRUCT_BODY()
+	GENERATED_BODY()
 };
 
 template<>
@@ -432,7 +432,7 @@ struct TStructOpsTypeTraits<FAutomationWorkerNextNetworkCommandReply> : public T
 USTRUCT()
 struct FAutomationWorkerScreenImage
 {
-	GENERATED_USTRUCT_BODY()
+	GENERATED_BODY()
 
 	/** The screen shot data. */
 	UPROPERTY()
diff --git a/Engine/Source/Runtime/CoreUObject/Public/UObject/ObjectBase.h b/Engine/Source/Runtime/CoreUObject/Public/UObject/ObjectBase.h
--- a/Engine/Source/Runtime/CoreUObject/Public/UObject/ObjectBase.h
+++ b/Engine/Source/Runtime/CoreUObject/Public/UObject/ObjectBase.h
@@ -515,7 +515,7 @@ struct COREUOBJECT_API FReferencerInformationList
 #define UENUM(...)
 #define UDELEGATE(...)
 
-// This pair of macros is used to help implement GENERATED_BODY() and GENERATED_USTRUCT_BODY()
+// This pair of macros is used to help implement GENERATED*_BODY() macros
 #define BODY_MACRO_COMBINE_INNER(A,B,C,D) A##B##C##D
 #define BODY_MACRO_COMBINE(A,B,C,D) BODY_MACRO_COMBINE_INNER(A,B,C,D)
 
@@ -525,7 +525,7 @@ struct COREUOBJECT_API FReferencerInformationList
 #define GENERATED_USTRUCT_BODY(...) GENERATED_BODY()
 #define GENERATED_UCLASS_BODY(...) GENERATED_BODY_LEGACY()
 #define GENERATED_UINTERFACE_BODY(...) GENERATED_BODY_LEGACY()
-#define GENERATED_IINTERFACE_BODY(...) GENERATED_BODY()
+#define GENERATED_IINTERFACE_BODY(...) GENERATED_BODY_LEGACY()
 
 #if UE_BUILD_DOCS
 #define UCLASS(...)
diff --git a/Engine/Source/Runtime/Engine/Classes/AI/AISystemBase.h b/Engine/Source/Runtime/Engine/Classes/AI/AISystemBase.h
--- a/Engine/Source/Runtime/Engine/Classes/AI/AISystemBase.h
+++ b/Engine/Source/Runtime/Engine/Classes/AI/AISystemBase.h
@@ -9,7 +9,9 @@ class UWorld;
 UCLASS(abstract, config = Engine, defaultconfig)
 class ENGINE_API UAISystemBase : public UObject
 {
-	GENERATED_UCLASS_BODY()
+	GENERATED_BODY()
+public:
+	UAISystemBase(const FObjectInitializer& ObjectInitializer = FObjectInitializer::Get());
 
 	virtual ~UAISystemBase(){}
 		
@@ -50,4 +52,4 @@ class IAISystemModule : public IModuleInterface
 {
 public:
 	virtual UAISystemBase* CreateAISystemInstance(UWorld* World) = 0;
-};
\ No newline at end of file
+};
diff --git a/Engine/Source/Runtime/Engine/Classes/AI/Navigation/AbstractNavData.h b/Engine/Source/Runtime/Engine/Classes/AI/Navigation/AbstractNavData.h
--- a/Engine/Source/Runtime/Engine/Classes/AI/Navigation/AbstractNavData.h
+++ b/Engine/Source/Runtime/Engine/Classes/AI/Navigation/AbstractNavData.h
@@ -36,7 +36,9 @@ public:
 UCLASS()
 class AAbstractNavData : public ANavigationData
 {
-	GENERATED_UCLASS_BODY()
+	GENERATED_BODY()
+public:
+	AAbstractNavData(const FObjectInitializer& ObjectInitializer = FObjectInitializer::Get());
 
 	// Begin ANavigationData overrides
 	virtual void BatchRaycast(TArray<FNavigationRaycastWork>& Workload, TSharedPtr<const FNavigationQueryFilter> QueryFilter, const UObject* Querier = NULL) const override {};
diff --git a/Engine/Source/Runtime/Engine/Classes/AI/Navigation/AvoidanceManager.h b/Engine/Source/Runtime/Engine/Classes/AI/Navigation/AvoidanceManager.h
--- a/Engine/Source/Runtime/Engine/Classes/AI/Navigation/AvoidanceManager.h
+++ b/Engine/Source/Runtime/Engine/Classes/AI/Navigation/AvoidanceManager.h
@@ -12,7 +12,7 @@ class UAvoidanceManager;
 USTRUCT()
 struct FNavAvoidanceData
 {
-	GENERATED_USTRUCT_BODY()
+	GENERATED_BODY()
 
 	/** Current location */
 	FVector Center;
@@ -75,7 +75,9 @@ struct FVelocityAvoidanceCone
 UCLASS(config=Engine, Blueprintable)
 class ENGINE_API UAvoidanceManager : public UObject, public FSelfRegisteringExec
 {
-	GENERATED_UCLASS_BODY()
+	GENERATED_BODY()
+public:
+	UAvoidanceManager(const FObjectInitializer& ObjectInitializer = FObjectInitializer::Get());
 		
 	/** How long an avoidance UID must not be updated before the system will put it back in the pool. Actual delay is up to 150% of this value. */
 	UPROPERTY(EditAnywhere, Category="Avoidance", config, meta=(ClampMin = "0.0"))
diff --git a/Engine/Source/Runtime/Engine/Classes/AI/Navigation/NavAgentInterface.h b/Engine/Source/Runtime/Engine/Classes/AI/Navigation/NavAgentInterface.h
--- a/Engine/Source/Runtime/Engine/Classes/AI/Navigation/NavAgentInterface.h
+++ b/Engine/Source/Runtime/Engine/Classes/AI/Navigation/NavAgentInterface.h
@@ -11,12 +11,15 @@ struct FNavAgentProperties;
 UINTERFACE(MinimalAPI, meta=(CannotImplementInterfaceInBlueprint))
 class UNavAgentInterface : public UInterface
 {
-	GENERATED_UINTERFACE_BODY()
+	GENERATED_BODY()
+public:
+	ENGINE_API UNavAgentInterface(const FObjectInitializer& ObjectInitializer = FObjectInitializer::Get());
 };
 
 class INavAgentInterface
 {
-	GENERATED_IINTERFACE_BODY()
+	GENERATED_BODY()
+public:
 
 	/**
 	 *	Retrieves FNavAgentProperties expressing navigation props and caps of represented agent
diff --git a/Engine/Source/Runtime/Engine/Classes/AI/Navigation/NavAreas/NavArea.h b/Engine/Source/Runtime/Engine/Classes/AI/Navigation/NavAreas/NavArea.h
--- a/Engine/Source/Runtime/Engine/Classes/AI/Navigation/NavAreas/NavArea.h
+++ b/Engine/Source/Runtime/Engine/Classes/AI/Navigation/NavAreas/NavArea.h
@@ -7,7 +7,9 @@
 UCLASS(DefaultToInstanced, abstract, Config=Engine, Blueprintable)
 class ENGINE_API UNavArea : public UObject
 {
-	GENERATED_UCLASS_BODY()
+	GENERATED_BODY()
+public:
+	UNavArea(const FObjectInitializer& ObjectInitializer = FObjectInitializer::Get());
 
 	/** travel cost multiplier for path distance */
 	UPROPERTY(EditAnywhere, Category=NavArea, config, meta=(ClampMin = "0.0"))
diff --git a/Engine/Source/Runtime/Engine/Classes/AI/Navigation/NavAreas/NavAreaMeta.h b/Engine/Source/Runtime/Engine/Classes/AI/Navigation/NavAreas/NavAreaMeta.h
--- a/Engine/Source/Runtime/Engine/Classes/AI/Navigation/NavAreas/NavAreaMeta.h
+++ b/Engine/Source/Runtime/Engine/Classes/AI/Navigation/NavAreas/NavAreaMeta.h
@@ -11,7 +11,9 @@ class AActor;
 UCLASS(Abstract)
 class ENGINE_API UNavAreaMeta : public UNavArea
 {
-	GENERATED_UCLASS_BODY()
+	GENERATED_BODY()
+public:
+	UNavAreaMeta(const FObjectInitializer& ObjectInitializer = FObjectInitializer::Get());
 
 	/** 
 	 *	Picks an navigation area class that should be used for Actor when
diff --git a/Engine/Source/Runtime/Engine/Classes/AI/Navigation/NavAreas/NavAreaMeta_SwitchByAgent.h b/Engine/Source/Runtime/Engine/Classes/AI/Navigation/NavAreas/NavAreaMeta_SwitchByAgent.h
--- a/Engine/Source/Runtime/Engine/Classes/AI/Navigation/NavAreas/NavAreaMeta_SwitchByAgent.h
+++ b/Engine/Source/Runtime/Engine/Classes/AI/Navigation/NavAreas/NavAreaMeta_SwitchByAgent.h
@@ -9,7 +9,9 @@
 UCLASS(Abstract)
 class ENGINE_API UNavAreaMeta_SwitchByAgent : public UNavAreaMeta
 {
-	GENERATED_UCLASS_BODY()
+	GENERATED_BODY()
+public:
+	UNavAreaMeta_SwitchByAgent(const FObjectInitializer& ObjectInitializer = FObjectInitializer::Get());
 
 	UPROPERTY(EditAnywhere, Category=AgentTypes)
 	TSubclassOf<UNavArea> Agent0Area;
diff --git a/Engine/Source/Runtime/Engine/Classes/AI/Navigation/NavAreas/NavArea_Default.h b/Engine/Source/Runtime/Engine/Classes/AI/Navigation/NavAreas/NavArea_Default.h
--- a/Engine/Source/Runtime/Engine/Classes/AI/Navigation/NavAreas/NavArea_Default.h
+++ b/Engine/Source/Runtime/Engine/Classes/AI/Navigation/NavAreas/NavArea_Default.h
@@ -8,5 +8,7 @@
 UCLASS(Config=Engine)
 class ENGINE_API UNavArea_Default : public UNavArea
 {
-	GENERATED_UCLASS_BODY()
+	GENERATED_BODY()
+public:
+	UNavArea_Default(const FObjectInitializer& ObjectInitializer = FObjectInitializer::Get());
 };
diff --git a/Engine/Source/Runtime/Engine/Classes/AI/Navigation/NavAreas/NavArea_LowHeight.h b/Engine/Source/Runtime/Engine/Classes/AI/Navigation/NavAreas/NavArea_LowHeight.h
--- a/Engine/Source/Runtime/Engine/Classes/AI/Navigation/NavAreas/NavArea_LowHeight.h
+++ b/Engine/Source/Runtime/Engine/Classes/AI/Navigation/NavAreas/NavArea_LowHeight.h
@@ -8,5 +8,7 @@
 UCLASS(Config = Engine)
 class ENGINE_API UNavArea_LowHeight : public UNavArea
 {
-	GENERATED_UCLASS_BODY()
+	GENERATED_BODY()
+public:
+	UNavArea_LowHeight(const FObjectInitializer& ObjectInitializer = FObjectInitializer::Get());
 };
diff --git a/Engine/Source/Runtime/Engine/Classes/AI/Navigation/NavAreas/NavArea_Null.h b/Engine/Source/Runtime/Engine/Classes/AI/Navigation/NavAreas/NavArea_Null.h
--- a/Engine/Source/Runtime/Engine/Classes/AI/Navigation/NavAreas/NavArea_Null.h
+++ b/Engine/Source/Runtime/Engine/Classes/AI/Navigation/NavAreas/NavArea_Null.h
@@ -8,5 +8,7 @@
 UCLASS(Config=Engine)
 class ENGINE_API UNavArea_Null : public UNavArea
 {
-	GENERATED_UCLASS_BODY()
+	GENERATED_BODY()
+public:
+	UNavArea_Null(const FObjectInitializer& ObjectInitializer = FObjectInitializer::Get());
 };
diff --git a/Engine/Source/Runtime/Engine/Classes/AI/Navigation/NavAreas/NavArea_Obstacle.h b/Engine/Source/Runtime/Engine/Classes/AI/Navigation/NavAreas/NavArea_Obstacle.h
--- a/Engine/Source/Runtime/Engine/Classes/AI/Navigation/NavAreas/NavArea_Obstacle.h
+++ b/Engine/Source/Runtime/Engine/Classes/AI/Navigation/NavAreas/NavArea_Obstacle.h
@@ -8,5 +8,7 @@
 UCLASS(Config = Engine)
 class ENGINE_API UNavArea_Obstacle : public UNavArea
 {
-	GENERATED_UCLASS_BODY()
+	GENERATED_BODY()
+public:
+	UNavArea_Obstacle(const FObjectInitializer& ObjectInitializer = FObjectInitializer::Get());
 };
diff --git a/Engine/Source/Runtime/Engine/Classes/AI/Navigation/NavCollision.h b/Engine/Source/Runtime/Engine/Classes/AI/Navigation/NavCollision.h
--- a/Engine/Source/Runtime/Engine/Classes/AI/Navigation/NavCollision.h
+++ b/Engine/Source/Runtime/Engine/Classes/AI/Navigation/NavCollision.h
@@ -11,7 +11,7 @@ struct FCompositeNavModifier;
 USTRUCT()
 struct FNavCollisionCylinder
 {
-	GENERATED_USTRUCT_BODY()
+	GENERATED_BODY()
 
 	UPROPERTY(EditAnywhere, Category=Cylinder)
 	FVector Offset;
@@ -26,7 +26,7 @@ struct FNavCollisionCylinder
 USTRUCT()
 struct FNavCollisionBox
 {
-	GENERATED_USTRUCT_BODY()
+	GENERATED_BODY()
 
 	UPROPERTY(EditAnywhere, Category=Box)
 	FVector Offset;
@@ -44,7 +44,9 @@ struct FNavCollisionConvex
 UCLASS(config=Engine)
 class ENGINE_API UNavCollision : public UObject
 {
-	GENERATED_UCLASS_BODY()
+	GENERATED_BODY()
+public:
+	UNavCollision(const FObjectInitializer& ObjectInitializer = FObjectInitializer::Get());
 
 	FNavCollisionConvex TriMeshCollision;
 	FNavCollisionConvex ConvexCollision;
diff --git a/Engine/Source/Runtime/Engine/Classes/AI/Navigation/NavFilters/NavigationQueryFilter.h b/Engine/Source/Runtime/Engine/Classes/AI/Navigation/NavFilters/NavigationQueryFilter.h
--- a/Engine/Source/Runtime/Engine/Classes/AI/Navigation/NavFilters/NavigationQueryFilter.h
+++ b/Engine/Source/Runtime/Engine/Classes/AI/Navigation/NavFilters/NavigationQueryFilter.h
@@ -10,7 +10,7 @@ struct FNavigationQueryFilter;
 USTRUCT()
 struct ENGINE_API FNavigationFilterArea
 {
-	GENERATED_USTRUCT_BODY()
+	GENERATED_BODY()
 
 	/** navigation area class */
 	UPROPERTY(EditAnywhere, Category=Area)
@@ -43,7 +43,7 @@ struct ENGINE_API FNavigationFilterArea
 USTRUCT()
 struct ENGINE_API FNavigationFilterFlags
 {
-	GENERATED_USTRUCT_BODY()
+	GENERATED_BODY()
 
 #if CPP
 	union
@@ -201,7 +201,9 @@ protected:
 UCLASS(Abstract, Blueprintable)
 class ENGINE_API UNavigationQueryFilter : public UObject
 {
-	GENERATED_UCLASS_BODY()
+	GENERATED_BODY()
+public:
+	UNavigationQueryFilter(const FObjectInitializer& ObjectInitializer = FObjectInitializer::Get());
 	
 	/** list of overrides for navigation areas */
 	UPROPERTY(EditAnywhere, Category=Filter)
diff --git a/Engine/Source/Runtime/Engine/Classes/AI/Navigation/NavFilters/RecastFilter_UseDefaultArea.h b/Engine/Source/Runtime/Engine/Classes/AI/Navigation/NavFilters/RecastFilter_UseDefaultArea.h
--- a/Engine/Source/Runtime/Engine/Classes/AI/Navigation/NavFilters/RecastFilter_UseDefaultArea.h
+++ b/Engine/Source/Runtime/Engine/Classes/AI/Navigation/NavFilters/RecastFilter_UseDefaultArea.h
@@ -7,7 +7,9 @@
 UCLASS(MinimalAPI)
 class URecastFilter_UseDefaultArea : public UNavigationQueryFilter
 {
-	GENERATED_UCLASS_BODY()
+	GENERATED_BODY()
+public:
+	ENGINE_API URecastFilter_UseDefaultArea(const FObjectInitializer& ObjectInitializer = FObjectInitializer::Get());
 
 	virtual void InitializeFilter(const ANavigationData* NavData, FNavigationQueryFilter* Filter) const override;
 };
diff --git a/Engine/Source/Runtime/Engine/Classes/AI/Navigation/NavLinkCustomComponent.h b/Engine/Source/Runtime/Engine/Classes/AI/Navigation/NavLinkCustomComponent.h
--- a/Engine/Source/Runtime/Engine/Classes/AI/Navigation/NavLinkCustomComponent.h
+++ b/Engine/Source/Runtime/Engine/Classes/AI/Navigation/NavLinkCustomComponent.h
@@ -24,7 +24,9 @@ struct FNavigationRelevantData;
 UCLASS()
 class ENGINE_API UNavLinkCustomComponent : public UNavRelevantComponent, public INavLinkCustomInterface
 {
-	GENERATED_UCLASS_BODY()
+	GENERATED_BODY()
+public:
+	UNavLinkCustomComponent(const FObjectInitializer& ObjectInitializer = FObjectInitializer::Get());
 
 	DECLARE_DELEGATE_ThreeParams(FOnMoveReachedLink, UNavLinkCustomComponent* /*ThisComp*/, UPathFollowingComponent* /*PathComp*/, const FVector& /*DestPoint*/);
 	DECLARE_DELEGATE_TwoParams(FBroadcastFilter, UNavLinkCustomComponent* /*ThisComp*/, TArray<UPathFollowingComponent*>& /*NotifyList*/);
diff --git a/Engine/Source/Runtime/Engine/Classes/AI/Navigation/NavLinkCustomInterface.h b/Engine/Source/Runtime/Engine/Classes/AI/Navigation/NavLinkCustomInterface.h
--- a/Engine/Source/Runtime/Engine/Classes/AI/Navigation/NavLinkCustomInterface.h
+++ b/Engine/Source/Runtime/Engine/Classes/AI/Navigation/NavLinkCustomInterface.h
@@ -24,12 +24,15 @@ struct FNavigationLink;
 UINTERFACE(MinimalAPI, meta=(CannotImplementInterfaceInBlueprint))
 class UNavLinkCustomInterface : public UInterface
 {
-	GENERATED_UINTERFACE_BODY()
+	GENERATED_BODY()
+public:
+	ENGINE_API UNavLinkCustomInterface(const FObjectInitializer& ObjectInitializer = FObjectInitializer::Get());
 };
 
 class ENGINE_API INavLinkCustomInterface
 {
-	GENERATED_IINTERFACE_BODY()
+	GENERATED_BODY()
+public:
 
 	/** Get basic link data: two points (relative to owner) and direction */
 	virtual void GetLinkData(FVector& LeftPt, FVector& RightPt, ENavLinkDirection::Type& Direction) const {};
diff --git a/Engine/Source/Runtime/Engine/Classes/AI/Navigation/NavLinkDefinition.h b/Engine/Source/Runtime/Engine/Classes/AI/Navigation/NavLinkDefinition.h
--- a/Engine/Source/Runtime/Engine/Classes/AI/Navigation/NavLinkDefinition.h
+++ b/Engine/Source/Runtime/Engine/Classes/AI/Navigation/NavLinkDefinition.h
@@ -19,7 +19,7 @@ namespace ENavLinkDirection
 USTRUCT()
 struct ENGINE_API FNavigationLinkBase
 {
-	GENERATED_USTRUCT_BODY()
+	GENERATED_BODY()
 
 	/** if greater than 0 nav system will attempt to project navlink's end point geometry below */
 	UPROPERTY(EditAnywhere, Category=Default, meta=(ClampMin = "0.0"))
@@ -100,7 +100,7 @@ struct ENGINE_API FNavigationLinkBase
 USTRUCT()
 struct ENGINE_API FNavigationLink : public FNavigationLinkBase
 {
-	GENERATED_USTRUCT_BODY()
+	GENERATED_BODY()
 
 	UPROPERTY(EditAnywhere, Category=Default, BlueprintReadWrite, meta=(MakeEditWidget=""))
 	FVector Left;
@@ -148,7 +148,7 @@ struct ENGINE_API FNavigationLink : public FNavigationLinkBase
 USTRUCT()
 struct ENGINE_API FNavigationSegmentLink : public FNavigationLinkBase
 {
-	GENERATED_USTRUCT_BODY()
+	GENERATED_BODY()
 
 	UPROPERTY(EditAnywhere, Category=Default, BlueprintReadWrite, meta=(MakeEditWidget=""))
 	FVector LeftStart;
@@ -197,7 +197,9 @@ struct ENGINE_API FNavigationSegmentLink : public FNavigationLinkBase
 UCLASS(abstract, Config=Engine, Blueprintable)
 class ENGINE_API UNavLinkDefinition : public UObject
 {
-	GENERATED_UCLASS_BODY()
+	GENERATED_BODY()
+public:
+	UNavLinkDefinition(const FObjectInitializer& ObjectInitializer = FObjectInitializer::Get());
 
 	UPROPERTY(EditAnywhere, Category=OffMeshLinks)
 	TArray<FNavigationLink> Links;
diff --git a/Engine/Source/Runtime/Engine/Classes/AI/Navigation/NavLinkHostInterface.h b/Engine/Source/Runtime/Engine/Classes/AI/Navigation/NavLinkHostInterface.h
--- a/Engine/Source/Runtime/Engine/Classes/AI/Navigation/NavLinkHostInterface.h
+++ b/Engine/Source/Runtime/Engine/Classes/AI/Navigation/NavLinkHostInterface.h
@@ -10,12 +10,15 @@ struct FNavigationSegmentLink;
 UINTERFACE(MinimalAPI, meta=(CannotImplementInterfaceInBlueprint))
 class UNavLinkHostInterface : public UInterface
 {
-	GENERATED_UINTERFACE_BODY()
+	GENERATED_BODY()
+public:
+	ENGINE_API UNavLinkHostInterface(const FObjectInitializer& ObjectInitializer = FObjectInitializer::Get());
 };
 
 class ENGINE_API INavLinkHostInterface
 {
-	GENERATED_IINTERFACE_BODY()
+	GENERATED_BODY()
+public:
 		
 	/**
 	 *	Retrieves UNavLinkDefinition derived UClasses hosted by this interface implementor
diff --git a/Engine/Source/Runtime/Engine/Classes/AI/Navigation/NavLinkProxy.h b/Engine/Source/Runtime/Engine/Classes/AI/Navigation/NavLinkProxy.h
--- a/Engine/Source/Runtime/Engine/Classes/AI/Navigation/NavLinkProxy.h
+++ b/Engine/Source/Runtime/Engine/Classes/AI/Navigation/NavLinkProxy.h
@@ -16,7 +16,9 @@ DECLARE_DYNAMIC_MULTICAST_DELEGATE_TwoParams( FSmartLinkReachedSignature, class
 UCLASS(Blueprintable)
 class ENGINE_API ANavLinkProxy : public AActor, public INavLinkHostInterface, public INavRelevantInterface
 {
-	GENERATED_UCLASS_BODY()
+	GENERATED_BODY()
+public:
+	ANavLinkProxy(const FObjectInitializer& ObjectInitializer = FObjectInitializer::Get());
 
 	/** Navigation links (point to point) added to navigation data */
 	UPROPERTY(EditAnywhere, Category=SimpleLink)
diff --git a/Engine/Source/Runtime/Engine/Classes/AI/Navigation/NavLinkRenderingComponent.h b/Engine/Source/Runtime/Engine/Classes/AI/Navigation/NavLinkRenderingComponent.h
--- a/Engine/Source/Runtime/Engine/Classes/AI/Navigation/NavLinkRenderingComponent.h
+++ b/Engine/Source/Runtime/Engine/Classes/AI/Navigation/NavLinkRenderingComponent.h
@@ -9,7 +9,9 @@ class UNavLinkDefinition;
 UCLASS(hidecategories=Object, editinlinenew)
 class ENGINE_API UNavLinkRenderingComponent : public UPrimitiveComponent
 {
-	GENERATED_UCLASS_BODY()
+	GENERATED_BODY()
+public:
+	UNavLinkRenderingComponent(const FObjectInitializer& ObjectInitializer = FObjectInitializer::Get());
 		
 	// Begin UPrimitiveComponent Interface
 	virtual FPrimitiveSceneProxy* CreateSceneProxy() override;
diff --git a/Engine/Source/Runtime/Engine/Classes/AI/Navigation/NavLinkTrivial.h b/Engine/Source/Runtime/Engine/Classes/AI/Navigation/NavLinkTrivial.h
--- a/Engine/Source/Runtime/Engine/Classes/AI/Navigation/NavLinkTrivial.h
+++ b/Engine/Source/Runtime/Engine/Classes/AI/Navigation/NavLinkTrivial.h
@@ -6,5 +6,7 @@
 UCLASS()
 class ENGINE_API UNavLinkTrivial : public UNavLinkDefinition
 {
-	GENERATED_UCLASS_BODY()
-};
\ No newline at end of file
+	GENERATED_BODY()
+public:
+	UNavLinkTrivial(const FObjectInitializer& ObjectInitializer = FObjectInitializer::Get());
+};
diff --git a/Engine/Source/Runtime/Engine/Classes/AI/Navigation/NavMeshBoundsVolume.h b/Engine/Source/Runtime/Engine/Classes/AI/Navigation/NavMeshBoundsVolume.h
--- a/Engine/Source/Runtime/Engine/Classes/AI/Navigation/NavMeshBoundsVolume.h
+++ b/Engine/Source/Runtime/Engine/Classes/AI/Navigation/NavMeshBoundsVolume.h
@@ -12,7 +12,9 @@
 UCLASS(MinimalAPI)
 class ANavMeshBoundsVolume : public AVolume
 {
-	GENERATED_UCLASS_BODY()
+	GENERATED_BODY()
+public:
+	ENGINE_API ANavMeshBoundsVolume(const FObjectInitializer& ObjectInitializer = FObjectInitializer::Get());
 
 	// Begin AActor Interface
 	virtual void PostRegisterAllComponents() override;
diff --git a/Engine/Source/Runtime/Engine/Classes/AI/Navigation/NavMeshRenderingComponent.h b/Engine/Source/Runtime/Engine/Classes/AI/Navigation/NavMeshRenderingComponent.h
--- a/Engine/Source/Runtime/Engine/Classes/AI/Navigation/NavMeshRenderingComponent.h
+++ b/Engine/Source/Runtime/Engine/Classes/AI/Navigation/NavMeshRenderingComponent.h
@@ -8,7 +8,9 @@ struct FNavMeshSceneProxyData;
 UCLASS(hidecategories=Object, editinlinenew)
 class ENGINE_API UNavMeshRenderingComponent : public UPrimitiveComponent
 {
-	GENERATED_UCLASS_BODY()
+	GENERATED_BODY()
+public:
+	UNavMeshRenderingComponent(const FObjectInitializer& ObjectInitializer = FObjectInitializer::Get());
 
 public:
 	// Begin UObject Interface
diff --git a/Engine/Source/Runtime/Engine/Classes/AI/Navigation/NavModifierComponent.h b/Engine/Source/Runtime/Engine/Classes/AI/Navigation/NavModifierComponent.h
--- a/Engine/Source/Runtime/Engine/Classes/AI/Navigation/NavModifierComponent.h
+++ b/Engine/Source/Runtime/Engine/Classes/AI/Navigation/NavModifierComponent.h
@@ -11,7 +11,9 @@ struct FCompositeNavModifier;
 UCLASS(ClassGroup = (Navigation), meta = (BlueprintSpawnableComponent), hidecategories = (Activation))
 class UNavModifierComponent : public UNavRelevantComponent
 {
-	GENERATED_UCLASS_BODY()
+	GENERATED_BODY()
+public:
+	UNavModifierComponent(const FObjectInitializer& ObjectInitializer = FObjectInitializer::Get());
 
 	UPROPERTY(EditAnywhere, Category = Navigation)
 	TSubclassOf<UNavArea> AreaClass;
diff --git a/Engine/Source/Runtime/Engine/Classes/AI/Navigation/NavModifierVolume.h b/Engine/Source/Runtime/Engine/Classes/AI/Navigation/NavModifierVolume.h
--- a/Engine/Source/Runtime/Engine/Classes/AI/Navigation/NavModifierVolume.h
+++ b/Engine/Source/Runtime/Engine/Classes/AI/Navigation/NavModifierVolume.h
@@ -11,7 +11,9 @@
 UCLASS(hidecategories=(Navigation))
 class ANavModifierVolume : public AVolume, public INavRelevantInterface
 {
-	GENERATED_UCLASS_BODY()
+	GENERATED_BODY()
+public:
+	ANavModifierVolume(const FObjectInitializer& ObjectInitializer = FObjectInitializer::Get());
 
 	UPROPERTY(EditAnywhere, BlueprintReadWrite, Category=Default)
 	TSubclassOf<class UNavArea> AreaClass;
diff --git a/Engine/Source/Runtime/Engine/Classes/AI/Navigation/NavNodeInterface.h b/Engine/Source/Runtime/Engine/Classes/AI/Navigation/NavNodeInterface.h
--- a/Engine/Source/Runtime/Engine/Classes/AI/Navigation/NavNodeInterface.h
+++ b/Engine/Source/Runtime/Engine/Classes/AI/Navigation/NavNodeInterface.h
@@ -6,12 +6,15 @@
 UINTERFACE(MinimalAPI, meta=(CannotImplementInterfaceInBlueprint))
 class UNavNodeInterface : public UInterface
 {
-	GENERATED_UINTERFACE_BODY()
+	GENERATED_BODY()
+public:
+	ENGINE_API UNavNodeInterface(const FObjectInitializer& ObjectInitializer = FObjectInitializer::Get());
 };
 
 class INavNodeInterface
 {
-	GENERATED_IINTERFACE_BODY()
+	GENERATED_BODY()
+public:
 
 
 	/**
diff --git a/Engine/Source/Runtime/Engine/Classes/AI/Navigation/NavPathObserverInterface.h b/Engine/Source/Runtime/Engine/Classes/AI/Navigation/NavPathObserverInterface.h
--- a/Engine/Source/Runtime/Engine/Classes/AI/Navigation/NavPathObserverInterface.h
+++ b/Engine/Source/Runtime/Engine/Classes/AI/Navigation/NavPathObserverInterface.h
@@ -6,12 +6,15 @@
 UINTERFACE(MinimalAPI, meta=(CannotImplementInterfaceInBlueprint))
 class UNavPathObserverInterface : public UInterface
 {
-	GENERATED_UINTERFACE_BODY()
+	GENERATED_BODY()
+public:
+	ENGINE_API UNavPathObserverInterface(const FObjectInitializer& ObjectInitializer = FObjectInitializer::Get());
 };
 
 class INavPathObserverInterface
 {
-	GENERATED_IINTERFACE_BODY()
+	GENERATED_BODY()
+public:
 
 	/** 
 	 *	Called on Navigation Agent when given PathGenerator updated it's path.
diff --git a/Engine/Source/Runtime/Engine/Classes/AI/Navigation/NavRelevantComponent.h b/Engine/Source/Runtime/Engine/Classes/AI/Navigation/NavRelevantComponent.h
--- a/Engine/Source/Runtime/Engine/Classes/AI/Navigation/NavRelevantComponent.h
+++ b/Engine/Source/Runtime/Engine/Classes/AI/Navigation/NavRelevantComponent.h
@@ -8,7 +8,9 @@
 UCLASS()
 class ENGINE_API UNavRelevantComponent : public UActorComponent, public INavRelevantInterface
 {
-	GENERATED_UCLASS_BODY()
+	GENERATED_BODY()
+public:
+	UNavRelevantComponent(const FObjectInitializer& ObjectInitializer = FObjectInitializer::Get());
 
 	// Begin UActorComponent Interface
 	virtual void OnRegister() override;
diff --git a/Engine/Source/Runtime/Engine/Classes/AI/Navigation/NavRelevantInterface.h b/Engine/Source/Runtime/Engine/Classes/AI/Navigation/NavRelevantInterface.h
--- a/Engine/Source/Runtime/Engine/Classes/AI/Navigation/NavRelevantInterface.h
+++ b/Engine/Source/Runtime/Engine/Classes/AI/Navigation/NavRelevantInterface.h
@@ -11,12 +11,15 @@ struct FNavigableGeometryExport;
 UINTERFACE(MinimalAPI, meta=(CannotImplementInterfaceInBlueprint))
 class UNavRelevantInterface : public UInterface
 {
-	GENERATED_UINTERFACE_BODY()
+	GENERATED_BODY()
+public:
+	ENGINE_API UNavRelevantInterface(const FObjectInitializer& ObjectInitializer = FObjectInitializer::Get());
 };
 
 class INavRelevantInterface
 {
-	GENERATED_IINTERFACE_BODY()
+	GENERATED_BODY()
+public:
 
 	/** Prepare navigation modifiers */
 	virtual void GetNavigationData(FNavigationRelevantData& Data) const {}
diff --git a/Engine/Source/Runtime/Engine/Classes/AI/Navigation/NavTestRenderingComponent.h b/Engine/Source/Runtime/Engine/Classes/AI/Navigation/NavTestRenderingComponent.h
--- a/Engine/Source/Runtime/Engine/Classes/AI/Navigation/NavTestRenderingComponent.h
+++ b/Engine/Source/Runtime/Engine/Classes/AI/Navigation/NavTestRenderingComponent.h
@@ -6,7 +6,9 @@
 UCLASS(hidecategories=Object)
 class UNavTestRenderingComponent: public UPrimitiveComponent
 {
-	GENERATED_UCLASS_BODY()
+	GENERATED_BODY()
+public:
+	UNavTestRenderingComponent(const FObjectInitializer& ObjectInitializer = FObjectInitializer::Get());
 
 	virtual FPrimitiveSceneProxy* CreateSceneProxy() override;
 	virtual FBoxSphereBounds CalcBounds(const FTransform &LocalToWorld) const override;
diff --git a/Engine/Source/Runtime/Engine/Classes/AI/Navigation/NavigationData.h b/Engine/Source/Runtime/Engine/Classes/AI/Navigation/NavigationData.h
--- a/Engine/Source/Runtime/Engine/Classes/AI/Navigation/NavigationData.h
+++ b/Engine/Source/Runtime/Engine/Classes/AI/Navigation/NavigationData.h
@@ -17,7 +17,7 @@ class UWorld;
 USTRUCT()
 struct FSupportedAreaData
 {
-	GENERATED_USTRUCT_BODY()
+	GENERATED_BODY()
 
 	UPROPERTY()
 	FString AreaClassName;
@@ -410,7 +410,9 @@ enum class ERuntimeGenerationType : uint8
 UCLASS(config=Engine, defaultconfig, NotBlueprintable, abstract)
 class ENGINE_API ANavigationData : public AActor
 {
-	GENERATED_UCLASS_BODY()
+	GENERATED_BODY()
+public:
+	ANavigationData(const FObjectInitializer& ObjectInitializer = FObjectInitializer::Get());
 	
 	UPROPERTY(transient, duplicatetransient)
 	UPrimitiveComponent* RenderingComp;
diff --git a/Engine/Source/Runtime/Engine/Classes/AI/Navigation/NavigationDataChunk.h b/Engine/Source/Runtime/Engine/Classes/AI/Navigation/NavigationDataChunk.h
--- a/Engine/Source/Runtime/Engine/Classes/AI/Navigation/NavigationDataChunk.h
+++ b/Engine/Source/Runtime/Engine/Classes/AI/Navigation/NavigationDataChunk.h
@@ -10,7 +10,9 @@
 UCLASS(NotBlueprintable, abstract)
 class ENGINE_API UNavigationDataChunk : public UObject
 {
-	GENERATED_UCLASS_BODY()
+	GENERATED_BODY()
+public:
+	UNavigationDataChunk(const FObjectInitializer& ObjectInitializer = FObjectInitializer::Get());
 
 public:
 	/** Name of NavigationData actor that owns this chunk */
diff --git a/Engine/Source/Runtime/Engine/Classes/AI/Navigation/NavigationGraph.h b/Engine/Source/Runtime/Engine/Classes/AI/Navigation/NavigationGraph.h
--- a/Engine/Source/Runtime/Engine/Classes/AI/Navigation/NavigationGraph.h
+++ b/Engine/Source/Runtime/Engine/Classes/AI/Navigation/NavigationGraph.h
@@ -9,7 +9,7 @@ class UNavigationSystem;
 USTRUCT()
 struct FNavGraphEdge
 {
-	GENERATED_USTRUCT_BODY()
+	GENERATED_BODY()
 
 	enum {
 		EgdeFlagsCount = 7
@@ -35,7 +35,7 @@ struct FNavGraphEdge
 USTRUCT()
 struct FNavGraphNode
 {
-	GENERATED_USTRUCT_BODY()
+	GENERATED_BODY()
 
 	/** Who's this node referring to? This will most commonly point to an actor or a component */
 	UPROPERTY()
@@ -56,7 +56,9 @@ struct FNavGraphNode
 UCLASS(config=Engine, MinimalAPI, abstract)
 class ANavigationGraph : public ANavigationData
 {
-	GENERATED_UCLASS_BODY()
+	GENERATED_BODY()
+public:
+	ENGINE_API ANavigationGraph(const FObjectInitializer& ObjectInitializer = FObjectInitializer::Get());
 
 public:
 	static ANavigationData* CreateNavigationInstances(UNavigationSystem* NavSys);
diff --git a/Engine/Source/Runtime/Engine/Classes/AI/Navigation/NavigationGraphNode.h b/Engine/Source/Runtime/Engine/Classes/AI/Navigation/NavigationGraphNode.h
--- a/Engine/Source/Runtime/Engine/Classes/AI/Navigation/NavigationGraphNode.h
+++ b/Engine/Source/Runtime/Engine/Classes/AI/Navigation/NavigationGraphNode.h
@@ -9,5 +9,7 @@
 UCLASS(config=Engine, MinimalAPI, NotBlueprintable, abstract)
 class ANavigationGraphNode : public AActor
 {
-	GENERATED_UCLASS_BODY()
+	GENERATED_BODY()
+public:
+	ENGINE_API ANavigationGraphNode(const FObjectInitializer& ObjectInitializer = FObjectInitializer::Get());
 };
diff --git a/Engine/Source/Runtime/Engine/Classes/AI/Navigation/NavigationGraphNodeComponent.h b/Engine/Source/Runtime/Engine/Classes/AI/Navigation/NavigationGraphNodeComponent.h
--- a/Engine/Source/Runtime/Engine/Classes/AI/Navigation/NavigationGraphNodeComponent.h
+++ b/Engine/Source/Runtime/Engine/Classes/AI/Navigation/NavigationGraphNodeComponent.h
@@ -7,7 +7,9 @@
 UCLASS(config=Engine, MinimalAPI, HideCategories=(Mobility))
 class UNavigationGraphNodeComponent : public USceneComponent
 {
-	GENERATED_UCLASS_BODY()
+	GENERATED_BODY()
+public:
+	ENGINE_API UNavigationGraphNodeComponent(const FObjectInitializer& ObjectInitializer = FObjectInitializer::Get());
 
 	UPROPERTY()
 	FNavGraphNode Node;
diff --git a/Engine/Source/Runtime/Engine/Classes/AI/Navigation/NavigationPath.h b/Engine/Source/Runtime/Engine/Classes/AI/Navigation/NavigationPath.h
--- a/Engine/Source/Runtime/Engine/Classes/AI/Navigation/NavigationPath.h
+++ b/Engine/Source/Runtime/Engine/Classes/AI/Navigation/NavigationPath.h
@@ -16,7 +16,9 @@ DECLARE_DYNAMIC_MULTICAST_DELEGATE_TwoParams(FOnNavigationPathUpdated, UNavigati
 UCLASS(BlueprintType)
 class ENGINE_API UNavigationPath : public UObject
 {
-	GENERATED_UCLASS_BODY()
+	GENERATED_BODY()
+public:
+	UNavigationPath(const FObjectInitializer& ObjectInitializer = FObjectInitializer::Get());
 
 	UPROPERTY(BlueprintAssignable)
 	FOnNavigationPathUpdated PathUpdatedNotifier;
diff --git a/Engine/Source/Runtime/Engine/Classes/AI/Navigation/NavigationPathGenerator.h b/Engine/Source/Runtime/Engine/Classes/AI/Navigation/NavigationPathGenerator.h
--- a/Engine/Source/Runtime/Engine/Classes/AI/Navigation/NavigationPathGenerator.h
+++ b/Engine/Source/Runtime/Engine/Classes/AI/Navigation/NavigationPathGenerator.h
@@ -6,12 +6,15 @@
 UINTERFACE(MinimalAPI, meta=(CannotImplementInterfaceInBlueprint))
 class UNavigationPathGenerator : public UInterface
 {
-	GENERATED_UINTERFACE_BODY()
+	GENERATED_BODY()
+public:
+	ENGINE_API UNavigationPathGenerator(const FObjectInitializer& ObjectInitializer = FObjectInitializer::Get());
 };
 
 class INavigationPathGenerator
 {
-	GENERATED_IINTERFACE_BODY()
+	GENERATED_BODY()
+public:
 
 
 	/** 
diff --git a/Engine/Source/Runtime/Engine/Classes/AI/Navigation/NavigationSystem.h b/Engine/Source/Runtime/Engine/Classes/AI/Navigation/NavigationSystem.h
--- a/Engine/Source/Runtime/Engine/Classes/AI/Navigation/NavigationSystem.h
+++ b/Engine/Source/Runtime/Engine/Classes/AI/Navigation/NavigationSystem.h
@@ -127,7 +127,9 @@ private:
 UCLASS(Within=World, config=Engine, defaultconfig)
 class ENGINE_API UNavigationSystem : public UBlueprintFunctionLibrary
 {
-	GENERATED_UCLASS_BODY()
+	GENERATED_BODY()
+public:
+	UNavigationSystem(const FObjectInitializer& ObjectInitializer = FObjectInitializer::Get());
 
 	virtual ~UNavigationSystem();
 
diff --git a/Engine/Source/Runtime/Engine/Classes/AI/Navigation/NavigationTestingActor.h b/Engine/Source/Runtime/Engine/Classes/AI/Navigation/NavigationTestingActor.h
--- a/Engine/Source/Runtime/Engine/Classes/AI/Navigation/NavigationTestingActor.h
+++ b/Engine/Source/Runtime/Engine/Classes/AI/Navigation/NavigationTestingActor.h
@@ -33,7 +33,9 @@ namespace ENavCostDisplay
 UCLASS(hidecategories=(Object, Actor, Input, Rendering), showcategories=("Input|MouseInput", "Input|TouchInput"), Blueprintable)
 class ENGINE_API ANavigationTestingActor : public AActor, public INavAgentInterface, public INavPathObserverInterface
 {
-	GENERATED_UCLASS_BODY()
+	GENERATED_BODY()
+public:
+	ANavigationTestingActor(const FObjectInitializer& ObjectInitializer = FObjectInitializer::Get());
 
 private_subobject:
 	DEPRECATED_FORGAME(4.6, "CapsuleComponent should not be accessed directly, please use GetCapsuleComponent() function instead. CapsuleComponent will soon be private and your code will not compile.")
diff --git a/Engine/Source/Runtime/Engine/Classes/AI/Navigation/NavigationTypes.h b/Engine/Source/Runtime/Engine/Classes/AI/Navigation/NavigationTypes.h
--- a/Engine/Source/Runtime/Engine/Classes/AI/Navigation/NavigationTypes.h
+++ b/Engine/Source/Runtime/Engine/Classes/AI/Navigation/NavigationTypes.h
@@ -315,7 +315,7 @@ typedef TWeakPtr<struct FNavigationPath, ESPMode::ThreadSafe> FNavPathWeakPtr;
 USTRUCT()
 struct FMovementProperties
 {
-	GENERATED_USTRUCT_BODY()
+	GENERATED_BODY()
 
 	/** If true, this Pawn is capable of crouching. */
 	UPROPERTY(EditAnywhere, BlueprintReadWrite, Category=MovementProperties)
@@ -351,7 +351,7 @@ struct FMovementProperties
 USTRUCT()
 struct ENGINE_API FNavAgentProperties : public FMovementProperties
 {
-	GENERATED_USTRUCT_BODY()
+	GENERATED_BODY()
 
 	/** Radius of the capsule used for navigation/pathfinding. */
 	UPROPERTY(EditAnywhere, BlueprintReadWrite, Category = MovementProperties)
@@ -407,7 +407,7 @@ inline uint32 GetTypeHash(const FNavAgentProperties& A)
 USTRUCT()
 struct ENGINE_API FNavDataConfig : public FNavAgentProperties
 {
-	GENERATED_USTRUCT_BODY()
+	GENERATED_BODY()
 
 	UPROPERTY(EditAnywhere, BlueprintReadWrite, Category = Display)
 	FName Name;
@@ -520,10 +520,11 @@ struct ENGINE_API FMoveRequestCustomData
 typedef TSharedPtr<FMoveRequestCustomData, ESPMode::ThreadSafe> FCustomMoveSharedPtr;
 typedef TWeakPtr<FMoveRequestCustomData, ESPMode::ThreadSafe> FCustomMoveWeakPtr;
 
-UCLASS(Abstract, CustomConstructor)
+UCLASS(Abstract)
 class UNavigationTypes : public UObject
 {
-	GENERATED_UCLASS_BODY()
+	GENERATED_BODY()
+public:
 
 	UNavigationTypes(const FObjectInitializer& ObjectInitializer) : Super(ObjectInitializer) { }
 };
@@ -696,4 +697,4 @@ struct ENGINE_API FNavHeightfieldSamples
 	FNavHeightfieldSamples();
 
 	FORCEINLINE bool IsEmpty() const { return Heights.Num() == 0; }
-};
\ No newline at end of file
+};
diff --git a/Engine/Source/Runtime/Engine/Classes/AI/Navigation/RecastNavMesh.h b/Engine/Source/Runtime/Engine/Classes/AI/Navigation/RecastNavMesh.h
--- a/Engine/Source/Runtime/Engine/Classes/AI/Navigation/RecastNavMesh.h
+++ b/Engine/Source/Runtime/Engine/Classes/AI/Navigation/RecastNavMesh.h
@@ -411,7 +411,9 @@ namespace FNavMeshConfig
 UCLASS(config=Engine, defaultconfig, hidecategories=(Input,Rendering,Tags,Transform,"Utilities|Transformation",Actor,Layers,Replication), notplaceable)
 class ENGINE_API ARecastNavMesh : public ANavigationData
 {
-	GENERATED_UCLASS_BODY()
+	GENERATED_BODY()
+public:
+	ARecastNavMesh(const FObjectInitializer& ObjectInitializer = FObjectInitializer::Get());
 
 	typedef uint16 FNavPolyFlags;
 
diff --git a/Engine/Source/Runtime/Engine/Classes/AI/Navigation/RecastNavMeshDataChunk.h b/Engine/Source/Runtime/Engine/Classes/AI/Navigation/RecastNavMeshDataChunk.h
--- a/Engine/Source/Runtime/Engine/Classes/AI/Navigation/RecastNavMeshDataChunk.h
+++ b/Engine/Source/Runtime/Engine/Classes/AI/Navigation/RecastNavMeshDataChunk.h
@@ -36,7 +36,9 @@ class dtNavMesh;
 UCLASS()
 class ENGINE_API URecastNavMeshDataChunk : public UNavigationDataChunk
 {
-	GENERATED_UCLASS_BODY()
+	GENERATED_BODY()
+public:
+	URecastNavMeshDataChunk(const FObjectInitializer& ObjectInitializer = FObjectInitializer::Get());
 
 	// Begin UObject Interface
 	virtual void Serialize(FArchive& Ar) override;
diff --git a/Engine/Source/Runtime/Engine/Classes/Animation/AimOffsetBlendSpace.h b/Engine/Source/Runtime/Engine/Classes/Animation/AimOffsetBlendSpace.h
--- a/Engine/Source/Runtime/Engine/Classes/Animation/AimOffsetBlendSpace.h
+++ b/Engine/Source/Runtime/Engine/Classes/Animation/AimOffsetBlendSpace.h
@@ -16,7 +16,9 @@ struct FBlendSample;
 UCLASS(config=Engine, hidecategories=Object, MinimalAPI, BlueprintType)
 class UAimOffsetBlendSpace : public UBlendSpace
 {
-	GENERATED_UCLASS_BODY()
+	GENERATED_BODY()
+public:
+	ENGINE_API UAimOffsetBlendSpace(const FObjectInitializer& ObjectInitializer = FObjectInitializer::Get());
 
 	virtual bool IsValidAdditive() const override;
 
diff --git a/Engine/Source/Runtime/Engine/Classes/Animation/AimOffsetBlendSpace1D.h b/Engine/Source/Runtime/Engine/Classes/Animation/AimOffsetBlendSpace1D.h
--- a/Engine/Source/Runtime/Engine/Classes/Animation/AimOffsetBlendSpace1D.h
+++ b/Engine/Source/Runtime/Engine/Classes/Animation/AimOffsetBlendSpace1D.h
@@ -15,7 +15,9 @@
 UCLASS(config=Engine, hidecategories=Object, MinimalAPI, BlueprintType)
 class UAimOffsetBlendSpace1D : public UBlendSpace1D
 {
-	GENERATED_UCLASS_BODY()
+	GENERATED_BODY()
+public:
+	ENGINE_API UAimOffsetBlendSpace1D(const FObjectInitializer& ObjectInitializer = FObjectInitializer::Get());
 
 	virtual bool IsValidAdditive() const override;
 
diff --git a/Engine/Source/Runtime/Engine/Classes/Animation/AnimBlueprint.h b/Engine/Source/Runtime/Engine/Classes/Animation/AnimBlueprint.h
--- a/Engine/Source/Runtime/Engine/Classes/Animation/AnimBlueprint.h
+++ b/Engine/Source/Runtime/Engine/Classes/Animation/AnimBlueprint.h
@@ -7,7 +7,7 @@
 USTRUCT()
 struct FAnimGroupInfo
 {
-	GENERATED_USTRUCT_BODY()
+	GENERATED_BODY()
 
 	UPROPERTY()
 	FName Name;
@@ -24,7 +24,7 @@ struct FAnimGroupInfo
 USTRUCT()
 struct FAnimParentNodeAssetOverride
 {
-	GENERATED_USTRUCT_BODY()
+	GENERATED_BODY()
 
 	UPROPERTY()
 	UAnimationAsset* NewAsset;
@@ -54,7 +54,9 @@ struct FAnimParentNodeAssetOverride
 UCLASS(BlueprintType)
 class ENGINE_API UAnimBlueprint : public UBlueprint
 {
-	GENERATED_UCLASS_BODY()
+	GENERATED_BODY()
+public:
+	UAnimBlueprint(const FObjectInitializer& ObjectInitializer = FObjectInitializer::Get());
 
 	/** The kind of skeleton that animation graphs compiled from the blueprint will animate */
 	UPROPERTY(AssetRegistrySearchable)
diff --git a/Engine/Source/Runtime/Engine/Classes/Animation/AnimBlueprintGeneratedClass.h b/Engine/Source/Runtime/Engine/Classes/Animation/AnimBlueprintGeneratedClass.h
--- a/Engine/Source/Runtime/Engine/Classes/Animation/AnimBlueprintGeneratedClass.h
+++ b/Engine/Source/Runtime/Engine/Classes/Animation/AnimBlueprintGeneratedClass.h
@@ -19,7 +19,7 @@ class  UAnimGraphNode_StateResult;
 USTRUCT()
 struct FStateMachineDebugData
 {
-	GENERATED_USTRUCT_BODY()
+	GENERATED_BODY()
 
 public:
 	// Map from state nodes to their state entry in a state machine
@@ -41,7 +41,7 @@ public:
 USTRUCT()
 struct FAnimationFrameSnapshot
 {
-	GENERATED_USTRUCT_BODY()
+	GENERATED_BODY()
 
 	FAnimationFrameSnapshot()
 #if WITH_EDITORONLY_DATA
@@ -68,7 +68,7 @@ public:
 USTRUCT()
 struct ENGINE_API FAnimBlueprintDebugData
 {
-	GENERATED_USTRUCT_BODY()
+	GENERATED_BODY()
 
 	FAnimBlueprintDebugData()
 #if WITH_EDITORONLY_DATA
@@ -161,7 +161,9 @@ namespace EPropertySearchMode
 UCLASS()
 class ENGINE_API UAnimBlueprintGeneratedClass : public UBlueprintGeneratedClass
 {
-	GENERATED_UCLASS_BODY()
+	GENERATED_BODY()
+public:
+	UAnimBlueprintGeneratedClass(const FObjectInitializer& ObjectInitializer = FObjectInitializer::Get());
 
 	// List of state machines present in this blueprint class
 	UPROPERTY()
@@ -299,4 +301,4 @@ NodeType* GetNodeFromPropertyIndex(UAnimInstance* AnimInstance, const UAnimBluep
 	}
 
 	return NULL;
-}
\ No newline at end of file
+}
diff --git a/Engine/Source/Runtime/Engine/Classes/Animation/AnimComposite.h b/Engine/Source/Runtime/Engine/Classes/Animation/AnimComposite.h
--- a/Engine/Source/Runtime/Engine/Classes/Animation/AnimComposite.h
+++ b/Engine/Source/Runtime/Engine/Classes/Animation/AnimComposite.h
@@ -12,7 +12,9 @@
 UCLASS(config=Engine, hidecategories=UObject, MinimalAPI, BlueprintType)
 class UAnimComposite : public UAnimCompositeBase
 {
-	GENERATED_UCLASS_BODY()
+	GENERATED_BODY()
+public:
+	ENGINE_API UAnimComposite(const FObjectInitializer& ObjectInitializer = FObjectInitializer::Get());
 
 public:
 	/** Serializable data that stores section/anim pairing **/
