Back out changelist 2481333

diff --git a/Engine/Source/Runtime/AIModule/Classes/BehaviorTree/BehaviorTreeTypes.h b/Engine/Source/Runtime/AIModule/Classes/BehaviorTree/BehaviorTreeTypes.h
--- a/Engine/Source/Runtime/AIModule/Classes/BehaviorTree/BehaviorTreeTypes.h
+++ b/Engine/Source/Runtime/AIModule/Classes/BehaviorTree/BehaviorTreeTypes.h
@@ -426,7 +426,7 @@ struct FBehaviorTreePropertyMemory
 USTRUCT(BlueprintType)
 struct AIMODULE_API FBlackboardKeySelector
 {
-	GENERATED_BODY()
+	GENERATED_USTRUCT_BODY()
 
 	FBlackboardKeySelector() : SelectedKeyID(FBlackboard::InvalidKey)
 	{}
@@ -510,9 +510,7 @@ public:
 UCLASS(Abstract)
 class AIMODULE_API UBehaviorTreeTypes : public UObject
 {
-	GENERATED_BODY()
-public:
-	UBehaviorTreeTypes(const FObjectInitializer& ObjectInitializer = FObjectInitializer::Get());
+	GENERATED_UCLASS_BODY()
 
 	static FString DescribeNodeHelper(const UBTNode* Node);
 
diff --git a/Engine/Source/Runtime/AIModule/Classes/BehaviorTree/Blackboard/BlackboardKeyType.h b/Engine/Source/Runtime/AIModule/Classes/BehaviorTree/Blackboard/BlackboardKeyType.h
--- a/Engine/Source/Runtime/AIModule/Classes/BehaviorTree/Blackboard/BlackboardKeyType.h
+++ b/Engine/Source/Runtime/AIModule/Classes/BehaviorTree/Blackboard/BlackboardKeyType.h
@@ -77,9 +77,7 @@ struct FBlackboardInstancedKeyMemory
 UCLASS(EditInlineNew, Abstract, CollapseCategories, AutoExpandCategories=(Blackboard))
 class AIMODULE_API UBlackboardKeyType : public UObject
 {
-	GENERATED_BODY()
-public:
-	UBlackboardKeyType(const FObjectInitializer& ObjectInitializer = FObjectInitializer::Get());
+	GENERATED_UCLASS_BODY()
 
 	/** handle dynamic data size */
 	virtual void PreInitialize(UBlackboardComponent& OwnerComp);
diff --git a/Engine/Source/Runtime/AIModule/Classes/BehaviorTree/Blackboard/BlackboardKeyType_Bool.h b/Engine/Source/Runtime/AIModule/Classes/BehaviorTree/Blackboard/BlackboardKeyType_Bool.h
--- a/Engine/Source/Runtime/AIModule/Classes/BehaviorTree/Blackboard/BlackboardKeyType_Bool.h
+++ b/Engine/Source/Runtime/AIModule/Classes/BehaviorTree/Blackboard/BlackboardKeyType_Bool.h
@@ -6,9 +6,7 @@
 UCLASS(EditInlineNew, meta=(DisplayName="Bool"))
 class AIMODULE_API UBlackboardKeyType_Bool : public UBlackboardKeyType
 {
-	GENERATED_BODY()
-public:
-	UBlackboardKeyType_Bool(const FObjectInitializer& ObjectInitializer = FObjectInitializer::Get());
+	GENERATED_UCLASS_BODY()
 
 	typedef bool FDataType;
 	static const FDataType InvalidValue;
diff --git a/Engine/Source/Runtime/AIModule/Classes/BehaviorTree/Blackboard/BlackboardKeyType_Class.h b/Engine/Source/Runtime/AIModule/Classes/BehaviorTree/Blackboard/BlackboardKeyType_Class.h
--- a/Engine/Source/Runtime/AIModule/Classes/BehaviorTree/Blackboard/BlackboardKeyType_Class.h
+++ b/Engine/Source/Runtime/AIModule/Classes/BehaviorTree/Blackboard/BlackboardKeyType_Class.h
@@ -6,9 +6,7 @@
 UCLASS(EditInlineNew, meta=(DisplayName="Class"))
 class AIMODULE_API UBlackboardKeyType_Class : public UBlackboardKeyType
 {
-	GENERATED_BODY()
-public:
-	UBlackboardKeyType_Class(const FObjectInitializer& ObjectInitializer = FObjectInitializer::Get());
+	GENERATED_UCLASS_BODY()
 
 	typedef UClass* FDataType;
 	static const FDataType InvalidValue;
diff --git a/Engine/Source/Runtime/AIModule/Classes/BehaviorTree/Blackboard/BlackboardKeyType_Enum.h b/Engine/Source/Runtime/AIModule/Classes/BehaviorTree/Blackboard/BlackboardKeyType_Enum.h
--- a/Engine/Source/Runtime/AIModule/Classes/BehaviorTree/Blackboard/BlackboardKeyType_Enum.h
+++ b/Engine/Source/Runtime/AIModule/Classes/BehaviorTree/Blackboard/BlackboardKeyType_Enum.h
@@ -6,9 +6,7 @@
 UCLASS(EditInlineNew, meta=(DisplayName="Enum"))
 class AIMODULE_API UBlackboardKeyType_Enum : public UBlackboardKeyType
 {
-	GENERATED_BODY()
-public:
-	UBlackboardKeyType_Enum(const FObjectInitializer& ObjectInitializer = FObjectInitializer::Get());
+	GENERATED_UCLASS_BODY()
 
 	typedef uint8 FDataType;
 	static const FDataType InvalidValue;
diff --git a/Engine/Source/Runtime/AIModule/Classes/BehaviorTree/Blackboard/BlackboardKeyType_Float.h b/Engine/Source/Runtime/AIModule/Classes/BehaviorTree/Blackboard/BlackboardKeyType_Float.h
--- a/Engine/Source/Runtime/AIModule/Classes/BehaviorTree/Blackboard/BlackboardKeyType_Float.h
+++ b/Engine/Source/Runtime/AIModule/Classes/BehaviorTree/Blackboard/BlackboardKeyType_Float.h
@@ -6,9 +6,7 @@
 UCLASS(EditInlineNew, meta=(DisplayName="Float"))
 class AIMODULE_API UBlackboardKeyType_Float : public UBlackboardKeyType
 {
-	GENERATED_BODY()
-public:
-	UBlackboardKeyType_Float(const FObjectInitializer& ObjectInitializer = FObjectInitializer::Get());
+	GENERATED_UCLASS_BODY()
 
 	typedef float FDataType;
 	static const FDataType InvalidValue;
diff --git a/Engine/Source/Runtime/AIModule/Classes/BehaviorTree/Blackboard/BlackboardKeyType_Int.h b/Engine/Source/Runtime/AIModule/Classes/BehaviorTree/Blackboard/BlackboardKeyType_Int.h
--- a/Engine/Source/Runtime/AIModule/Classes/BehaviorTree/Blackboard/BlackboardKeyType_Int.h
+++ b/Engine/Source/Runtime/AIModule/Classes/BehaviorTree/Blackboard/BlackboardKeyType_Int.h
@@ -6,9 +6,7 @@
 UCLASS(EditInlineNew, meta=(DisplayName="Int"))
 class AIMODULE_API UBlackboardKeyType_Int : public UBlackboardKeyType
 {
-	GENERATED_BODY()
-public:
-	UBlackboardKeyType_Int(const FObjectInitializer& ObjectInitializer = FObjectInitializer::Get());
+	GENERATED_UCLASS_BODY()
 
 	typedef int32 FDataType;
 	static const FDataType InvalidValue;
diff --git a/Engine/Source/Runtime/AIModule/Classes/BehaviorTree/Blackboard/BlackboardKeyType_Name.h b/Engine/Source/Runtime/AIModule/Classes/BehaviorTree/Blackboard/BlackboardKeyType_Name.h
--- a/Engine/Source/Runtime/AIModule/Classes/BehaviorTree/Blackboard/BlackboardKeyType_Name.h
+++ b/Engine/Source/Runtime/AIModule/Classes/BehaviorTree/Blackboard/BlackboardKeyType_Name.h
@@ -6,9 +6,7 @@
 UCLASS(EditInlineNew, meta=(DisplayName="Name"))
 class AIMODULE_API UBlackboardKeyType_Name : public UBlackboardKeyType
 {
-	GENERATED_BODY()
-public:
-	UBlackboardKeyType_Name(const FObjectInitializer& ObjectInitializer = FObjectInitializer::Get());
+	GENERATED_UCLASS_BODY()
 
 	typedef FName FDataType;
 	static const FDataType InvalidValue;
diff --git a/Engine/Source/Runtime/AIModule/Classes/BehaviorTree/Blackboard/BlackboardKeyType_NativeEnum.h b/Engine/Source/Runtime/AIModule/Classes/BehaviorTree/Blackboard/BlackboardKeyType_NativeEnum.h
--- a/Engine/Source/Runtime/AIModule/Classes/BehaviorTree/Blackboard/BlackboardKeyType_NativeEnum.h
+++ b/Engine/Source/Runtime/AIModule/Classes/BehaviorTree/Blackboard/BlackboardKeyType_NativeEnum.h
@@ -9,9 +9,7 @@
 UCLASS(NotEditInlineNew)
 class AIMODULE_API UBlackboardKeyType_NativeEnum : public UBlackboardKeyType
 {
-	GENERATED_BODY()
-public:
-	UBlackboardKeyType_NativeEnum(const FObjectInitializer& ObjectInitializer = FObjectInitializer::Get());
+	GENERATED_UCLASS_BODY()
 
 	typedef uint8 FDataType;
 	static const FDataType InvalidValue;
diff --git a/Engine/Source/Runtime/AIModule/Classes/BehaviorTree/Blackboard/BlackboardKeyType_Object.h b/Engine/Source/Runtime/AIModule/Classes/BehaviorTree/Blackboard/BlackboardKeyType_Object.h
--- a/Engine/Source/Runtime/AIModule/Classes/BehaviorTree/Blackboard/BlackboardKeyType_Object.h
+++ b/Engine/Source/Runtime/AIModule/Classes/BehaviorTree/Blackboard/BlackboardKeyType_Object.h
@@ -6,9 +6,7 @@
 UCLASS(EditInlineNew, meta=(DisplayName="Object"))
 class AIMODULE_API UBlackboardKeyType_Object : public UBlackboardKeyType
 {
-	GENERATED_BODY()
-public:
-	UBlackboardKeyType_Object(const FObjectInitializer& ObjectInitializer = FObjectInitializer::Get());
+	GENERATED_UCLASS_BODY()
 
 	typedef UObject* FDataType;
 	static const FDataType InvalidValue;
diff --git a/Engine/Source/Runtime/AIModule/Classes/BehaviorTree/Blackboard/BlackboardKeyType_Rotator.h b/Engine/Source/Runtime/AIModule/Classes/BehaviorTree/Blackboard/BlackboardKeyType_Rotator.h
--- a/Engine/Source/Runtime/AIModule/Classes/BehaviorTree/Blackboard/BlackboardKeyType_Rotator.h
+++ b/Engine/Source/Runtime/AIModule/Classes/BehaviorTree/Blackboard/BlackboardKeyType_Rotator.h
@@ -5,9 +5,7 @@
 UCLASS(EditInlineNew, meta=(DisplayName="Rotator"))
 class AIMODULE_API UBlackboardKeyType_Rotator : public UBlackboardKeyType
 {
-	GENERATED_BODY()
-public:
-	UBlackboardKeyType_Rotator(const FObjectInitializer& ObjectInitializer = FObjectInitializer::Get());
+	GENERATED_UCLASS_BODY()
 	
 	typedef FRotator FDataType; 
 	static const FDataType InvalidValue;
diff --git a/Engine/Source/Runtime/AIModule/Classes/BehaviorTree/Blackboard/BlackboardKeyType_String.h b/Engine/Source/Runtime/AIModule/Classes/BehaviorTree/Blackboard/BlackboardKeyType_String.h
--- a/Engine/Source/Runtime/AIModule/Classes/BehaviorTree/Blackboard/BlackboardKeyType_String.h
+++ b/Engine/Source/Runtime/AIModule/Classes/BehaviorTree/Blackboard/BlackboardKeyType_String.h
@@ -6,9 +6,7 @@
 UCLASS(EditInlineNew, meta=(DisplayName="String"))
 class AIMODULE_API UBlackboardKeyType_String : public UBlackboardKeyType
 {
-	GENERATED_BODY()
-public:
-	UBlackboardKeyType_String(const FObjectInitializer& ObjectInitializer = FObjectInitializer::Get());
+	GENERATED_UCLASS_BODY()
 
 	typedef FString FDataType;
 	static const FDataType InvalidValue;
diff --git a/Engine/Source/Runtime/AIModule/Classes/BehaviorTree/Blackboard/BlackboardKeyType_Vector.h b/Engine/Source/Runtime/AIModule/Classes/BehaviorTree/Blackboard/BlackboardKeyType_Vector.h
--- a/Engine/Source/Runtime/AIModule/Classes/BehaviorTree/Blackboard/BlackboardKeyType_Vector.h
+++ b/Engine/Source/Runtime/AIModule/Classes/BehaviorTree/Blackboard/BlackboardKeyType_Vector.h
@@ -6,9 +6,7 @@
 UCLASS(EditInlineNew, meta=(DisplayName="Vector"))
 class AIMODULE_API UBlackboardKeyType_Vector : public UBlackboardKeyType
 {
-	GENERATED_BODY()
-public:
-	UBlackboardKeyType_Vector(const FObjectInitializer& ObjectInitializer = FObjectInitializer::Get());
+	GENERATED_UCLASS_BODY()
 
 	typedef FVector FDataType; 
 	static const FDataType InvalidValue;
diff --git a/Engine/Source/Runtime/AIModule/Classes/BehaviorTree/BlackboardData.h b/Engine/Source/Runtime/AIModule/Classes/BehaviorTree/BlackboardData.h
--- a/Engine/Source/Runtime/AIModule/Classes/BehaviorTree/BlackboardData.h
+++ b/Engine/Source/Runtime/AIModule/Classes/BehaviorTree/BlackboardData.h
@@ -11,7 +11,7 @@ class UBlackboardData;
 USTRUCT()
 struct FBlackboardEntry
 {
-	GENERATED_BODY()
+	GENERATED_USTRUCT_BODY()
 
 	UPROPERTY(EditAnywhere, Category=Blackboard)
 	FName EntryName;
@@ -35,9 +35,7 @@ struct FBlackboardEntry
 UCLASS(AutoExpandCategories=(Blackboard))
 class AIMODULE_API UBlackboardData : public UDataAsset
 {
-	GENERATED_BODY()
-public:
-	UBlackboardData(const FObjectInitializer& ObjectInitializer = FObjectInitializer::Get());
+	GENERATED_UCLASS_BODY()
 	DECLARE_MULTICAST_DELEGATE_OneParam(FKeyUpdate, UBlackboardData* /*asset*/);
 
 	/** parent blackboard (keys can be overridden) */
diff --git a/Engine/Source/Runtime/AIModule/Classes/BehaviorTree/Composites/BTComposite_Selector.h b/Engine/Source/Runtime/AIModule/Classes/BehaviorTree/Composites/BTComposite_Selector.h
--- a/Engine/Source/Runtime/AIModule/Classes/BehaviorTree/Composites/BTComposite_Selector.h
+++ b/Engine/Source/Runtime/AIModule/Classes/BehaviorTree/Composites/BTComposite_Selector.h
@@ -11,9 +11,7 @@
 UCLASS()
 class AIMODULE_API UBTComposite_Selector: public UBTCompositeNode
 {
-	GENERATED_BODY()
-public:
-	UBTComposite_Selector(const FObjectInitializer& ObjectInitializer = FObjectInitializer::Get());
+	GENERATED_UCLASS_BODY()
 
 	int32 GetNextChildHandler(struct FBehaviorTreeSearchData& SearchData, int32 PrevChild, EBTNodeResult::Type LastResult) const;
 
diff --git a/Engine/Source/Runtime/AIModule/Classes/BehaviorTree/Composites/BTComposite_Sequence.h b/Engine/Source/Runtime/AIModule/Classes/BehaviorTree/Composites/BTComposite_Sequence.h
--- a/Engine/Source/Runtime/AIModule/Classes/BehaviorTree/Composites/BTComposite_Sequence.h
+++ b/Engine/Source/Runtime/AIModule/Classes/BehaviorTree/Composites/BTComposite_Sequence.h
@@ -11,9 +11,7 @@
 UCLASS()
 class AIMODULE_API UBTComposite_Sequence : public UBTCompositeNode
 {
-	GENERATED_BODY()
-public:
-	UBTComposite_Sequence(const FObjectInitializer& ObjectInitializer = FObjectInitializer::Get());
+	GENERATED_UCLASS_BODY()
 
 	int32 GetNextChildHandler(struct FBehaviorTreeSearchData& SearchData, int32 PrevChild, EBTNodeResult::Type LastResult) const;
 
diff --git a/Engine/Source/Runtime/AIModule/Classes/BehaviorTree/Composites/BTComposite_SimpleParallel.h b/Engine/Source/Runtime/AIModule/Classes/BehaviorTree/Composites/BTComposite_SimpleParallel.h
--- a/Engine/Source/Runtime/AIModule/Classes/BehaviorTree/Composites/BTComposite_SimpleParallel.h
+++ b/Engine/Source/Runtime/AIModule/Classes/BehaviorTree/Composites/BTComposite_SimpleParallel.h
@@ -52,9 +52,7 @@ struct FBTParallelMemory : public FBTCompositeMemory
 UCLASS()
 class AIMODULE_API UBTComposite_SimpleParallel : public UBTCompositeNode
 {
-	GENERATED_BODY()
-public:
-	UBTComposite_SimpleParallel(const FObjectInitializer& ObjectInitializer = FObjectInitializer::Get());
+	GENERATED_UCLASS_BODY()
 
 	/** how background tree should be handled when main task finishes execution */
 	UPROPERTY(EditInstanceOnly, Category=Parallel)
diff --git a/Engine/Source/Runtime/AIModule/Classes/BehaviorTree/Decorators/BTDecorator_Blackboard.h b/Engine/Source/Runtime/AIModule/Classes/BehaviorTree/Decorators/BTDecorator_Blackboard.h
--- a/Engine/Source/Runtime/AIModule/Classes/BehaviorTree/Decorators/BTDecorator_Blackboard.h
+++ b/Engine/Source/Runtime/AIModule/Classes/BehaviorTree/Decorators/BTDecorator_Blackboard.h
@@ -31,9 +31,7 @@ namespace EBTBlackboardRestart
 UCLASS(HideCategories=(Condition))
 class AIMODULE_API UBTDecorator_Blackboard : public UBTDecorator_BlackboardBase
 {
-	GENERATED_BODY()
-public:
-	UBTDecorator_Blackboard(const FObjectInitializer& ObjectInitializer = FObjectInitializer::Get());
+	GENERATED_UCLASS_BODY()
 
 	virtual bool CalculateRawConditionValue(UBehaviorTreeComponent& OwnerComp, uint8* NodeMemory) const override;
 	virtual void OnBlackboardChange(const UBlackboardComponent& Blackboard, FBlackboard::FKey ChangedKeyID) override;
diff --git a/Engine/Source/Runtime/AIModule/Classes/BehaviorTree/Decorators/BTDecorator_BlackboardBase.h b/Engine/Source/Runtime/AIModule/Classes/BehaviorTree/Decorators/BTDecorator_BlackboardBase.h
--- a/Engine/Source/Runtime/AIModule/Classes/BehaviorTree/Decorators/BTDecorator_BlackboardBase.h
+++ b/Engine/Source/Runtime/AIModule/Classes/BehaviorTree/Decorators/BTDecorator_BlackboardBase.h
@@ -8,9 +8,7 @@ class UBlackboardComponent;
 UCLASS(Abstract)
 class AIMODULE_API UBTDecorator_BlackboardBase : public UBTDecorator
 {
-	GENERATED_BODY()
-public:
-	UBTDecorator_BlackboardBase(const FObjectInitializer& ObjectInitializer = FObjectInitializer::Get());
+	GENERATED_UCLASS_BODY()
 
 	/** initialize any asset related data */
 	virtual void InitializeFromAsset(UBehaviorTree& Asset) override;
diff --git a/Engine/Source/Runtime/AIModule/Classes/BehaviorTree/Decorators/BTDecorator_CheckGameplayTagsOnActor.h b/Engine/Source/Runtime/AIModule/Classes/BehaviorTree/Decorators/BTDecorator_CheckGameplayTagsOnActor.h
--- a/Engine/Source/Runtime/AIModule/Classes/BehaviorTree/Decorators/BTDecorator_CheckGameplayTagsOnActor.h
+++ b/Engine/Source/Runtime/AIModule/Classes/BehaviorTree/Decorators/BTDecorator_CheckGameplayTagsOnActor.h
@@ -14,9 +14,7 @@
 UCLASS()
 class AIMODULE_API UBTDecorator_CheckGameplayTagsOnActor : public UBTDecorator
 {
-	GENERATED_BODY()
-public:
-	UBTDecorator_CheckGameplayTagsOnActor(const FObjectInitializer& ObjectInitializer = FObjectInitializer::Get());
+	GENERATED_UCLASS_BODY()
 
 	virtual bool CalculateRawConditionValue(UBehaviorTreeComponent& OwnerComp, uint8* NodeMemory) const override;
 	virtual void DescribeRuntimeValues(const UBehaviorTreeComponent& OwnerComp, uint8* NodeMemory, EBTDescriptionVerbosity::Type Verbosity, TArray<FString>& Values) const override;
diff --git a/Engine/Source/Runtime/AIModule/Classes/BehaviorTree/Decorators/BTDecorator_CompareBBEntries.h b/Engine/Source/Runtime/AIModule/Classes/BehaviorTree/Decorators/BTDecorator_CompareBBEntries.h
--- a/Engine/Source/Runtime/AIModule/Classes/BehaviorTree/Decorators/BTDecorator_CompareBBEntries.h
+++ b/Engine/Source/Runtime/AIModule/Classes/BehaviorTree/Decorators/BTDecorator_CompareBBEntries.h
@@ -24,9 +24,7 @@ namespace EBlackBoardEntryComparison
 UCLASS(HideCategories=(Condition))
 class AIMODULE_API UBTDecorator_CompareBBEntries : public UBTDecorator
 {
-	GENERATED_BODY()
-public:
-	UBTDecorator_CompareBBEntries(const FObjectInitializer& ObjectInitializer = FObjectInitializer::Get());
+	GENERATED_UCLASS_BODY()
 
 protected:
 
diff --git a/Engine/Source/Runtime/AIModule/Classes/BehaviorTree/Decorators/BTDecorator_ConeCheck.h b/Engine/Source/Runtime/AIModule/Classes/BehaviorTree/Decorators/BTDecorator_ConeCheck.h
--- a/Engine/Source/Runtime/AIModule/Classes/BehaviorTree/Decorators/BTDecorator_ConeCheck.h
+++ b/Engine/Source/Runtime/AIModule/Classes/BehaviorTree/Decorators/BTDecorator_ConeCheck.h
@@ -19,9 +19,7 @@ struct FBTConeCheckDecoratorMemory
 UCLASS()
 class AIMODULE_API UBTDecorator_ConeCheck : public UBTDecorator
 {
-	GENERATED_BODY()
-public:
-	UBTDecorator_ConeCheck(const FObjectInitializer& ObjectInitializer = FObjectInitializer::Get());
+	GENERATED_UCLASS_BODY()
 
 	typedef FBTConeCheckDecoratorMemory TNodeInstanceMemory;
 
diff --git a/Engine/Source/Runtime/AIModule/Classes/BehaviorTree/Decorators/BTDecorator_Cooldown.h b/Engine/Source/Runtime/AIModule/Classes/BehaviorTree/Decorators/BTDecorator_Cooldown.h
--- a/Engine/Source/Runtime/AIModule/Classes/BehaviorTree/Decorators/BTDecorator_Cooldown.h
+++ b/Engine/Source/Runtime/AIModule/Classes/BehaviorTree/Decorators/BTDecorator_Cooldown.h
@@ -17,9 +17,7 @@ struct FBTCooldownDecoratorMemory
 UCLASS(HideCategories=(Condition))
 class AIMODULE_API UBTDecorator_Cooldown : public UBTDecorator
 {
-	GENERATED_BODY()
-public:
-	UBTDecorator_Cooldown(const FObjectInitializer& ObjectInitializer = FObjectInitializer::Get());
+	GENERATED_UCLASS_BODY()
 
 	/** max allowed time for execution of underlying node */
 	UPROPERTY(Category=Decorator, EditAnywhere)
diff --git a/Engine/Source/Runtime/AIModule/Classes/BehaviorTree/Decorators/BTDecorator_DoesPathExist.h b/Engine/Source/Runtime/AIModule/Classes/BehaviorTree/Decorators/BTDecorator_DoesPathExist.h
--- a/Engine/Source/Runtime/AIModule/Classes/BehaviorTree/Decorators/BTDecorator_DoesPathExist.h
+++ b/Engine/Source/Runtime/AIModule/Classes/BehaviorTree/Decorators/BTDecorator_DoesPathExist.h
@@ -24,9 +24,7 @@ namespace EPathExistanceQueryType
 UCLASS()
 class AIMODULE_API UBTDecorator_DoesPathExist : public UBTDecorator
 {
-	GENERATED_BODY()
-public:
-	UBTDecorator_DoesPathExist(const FObjectInitializer& ObjectInitializer = FObjectInitializer::Get());
+	GENERATED_UCLASS_BODY()
 
 protected:
 
diff --git a/Engine/Source/Runtime/AIModule/Classes/BehaviorTree/Decorators/BTDecorator_ForceSuccess.h b/Engine/Source/Runtime/AIModule/Classes/BehaviorTree/Decorators/BTDecorator_ForceSuccess.h
--- a/Engine/Source/Runtime/AIModule/Classes/BehaviorTree/Decorators/BTDecorator_ForceSuccess.h
+++ b/Engine/Source/Runtime/AIModule/Classes/BehaviorTree/Decorators/BTDecorator_ForceSuccess.h
@@ -16,9 +16,7 @@
 UCLASS(HideCategories=(Condition))
 class AIMODULE_API UBTDecorator_ForceSuccess : public UBTDecorator
 {
-	GENERATED_BODY()
-public:
-	UBTDecorator_ForceSuccess(const FObjectInitializer& ObjectInitializer = FObjectInitializer::Get());
+	GENERATED_UCLASS_BODY()
 
 protected:
 
diff --git a/Engine/Source/Runtime/AIModule/Classes/BehaviorTree/Decorators/BTDecorator_KeepInCone.h b/Engine/Source/Runtime/AIModule/Classes/BehaviorTree/Decorators/BTDecorator_KeepInCone.h
--- a/Engine/Source/Runtime/AIModule/Classes/BehaviorTree/Decorators/BTDecorator_KeepInCone.h
+++ b/Engine/Source/Runtime/AIModule/Classes/BehaviorTree/Decorators/BTDecorator_KeepInCone.h
@@ -16,9 +16,7 @@ struct FBTKeepInConeDecoratorMemory
 UCLASS(HideCategories=(Condition))
 class AIMODULE_API UBTDecorator_KeepInCone : public UBTDecorator
 {
-	GENERATED_BODY()
-public:
-	UBTDecorator_KeepInCone(const FObjectInitializer& ObjectInitializer = FObjectInitializer::Get());
+	GENERATED_UCLASS_BODY()
 
 	typedef FBTKeepInConeDecoratorMemory TNodeInstanceMemory;
 
diff --git a/Engine/Source/Runtime/AIModule/Classes/BehaviorTree/Decorators/BTDecorator_Loop.h b/Engine/Source/Runtime/AIModule/Classes/BehaviorTree/Decorators/BTDecorator_Loop.h
--- a/Engine/Source/Runtime/AIModule/Classes/BehaviorTree/Decorators/BTDecorator_Loop.h
+++ b/Engine/Source/Runtime/AIModule/Classes/BehaviorTree/Decorators/BTDecorator_Loop.h
@@ -16,9 +16,7 @@ struct FBTLoopDecoratorMemory
 UCLASS(HideCategories=(Condition))
 class AIMODULE_API UBTDecorator_Loop : public UBTDecorator
 {
-	GENERATED_BODY()
-public:
-	UBTDecorator_Loop(const FObjectInitializer& ObjectInitializer = FObjectInitializer::Get());
+	GENERATED_UCLASS_BODY()
 
 	/** number of executions */
 	UPROPERTY(Category=Decorator, EditAnywhere)
diff --git a/Engine/Source/Runtime/AIModule/Classes/BehaviorTree/Decorators/BTDecorator_TagCooldown.h b/Engine/Source/Runtime/AIModule/Classes/BehaviorTree/Decorators/BTDecorator_TagCooldown.h
--- a/Engine/Source/Runtime/AIModule/Classes/BehaviorTree/Decorators/BTDecorator_TagCooldown.h
+++ b/Engine/Source/Runtime/AIModule/Classes/BehaviorTree/Decorators/BTDecorator_TagCooldown.h
@@ -17,9 +17,7 @@ struct FBTTagCooldownDecoratorMemory
 UCLASS(HideCategories=(Condition))
 class AIMODULE_API UBTDecorator_TagCooldown : public UBTDecorator
 {
-	GENERATED_BODY()
-public:
-	UBTDecorator_TagCooldown(const FObjectInitializer& ObjectInitializer = FObjectInitializer::Get());
+	GENERATED_UCLASS_BODY()
 
 	/** gameplay tag that will be used for the cooldown */
 	UPROPERTY(Category = Decorator, EditAnywhere)
diff --git a/Engine/Source/Runtime/AIModule/Classes/BehaviorTree/Decorators/BTDecorator_TimeLimit.h b/Engine/Source/Runtime/AIModule/Classes/BehaviorTree/Decorators/BTDecorator_TimeLimit.h
--- a/Engine/Source/Runtime/AIModule/Classes/BehaviorTree/Decorators/BTDecorator_TimeLimit.h
+++ b/Engine/Source/Runtime/AIModule/Classes/BehaviorTree/Decorators/BTDecorator_TimeLimit.h
@@ -10,9 +10,7 @@
 UCLASS(HideCategories=(Condition))
 class AIMODULE_API UBTDecorator_TimeLimit : public UBTDecorator
 {
-	GENERATED_BODY()
-public:
-	UBTDecorator_TimeLimit(const FObjectInitializer& ObjectInitializer = FObjectInitializer::Get());
+	GENERATED_UCLASS_BODY()
 		
 	/** max allowed time for execution of underlying node */
 	UPROPERTY(Category=Decorator, EditAnywhere)
diff --git a/Engine/Source/Runtime/AIModule/Classes/BehaviorTree/Services/BTService_BlackboardBase.h b/Engine/Source/Runtime/AIModule/Classes/BehaviorTree/Services/BTService_BlackboardBase.h
--- a/Engine/Source/Runtime/AIModule/Classes/BehaviorTree/Services/BTService_BlackboardBase.h
+++ b/Engine/Source/Runtime/AIModule/Classes/BehaviorTree/Services/BTService_BlackboardBase.h
@@ -6,9 +6,7 @@
 UCLASS(Abstract)
 class AIMODULE_API UBTService_BlackboardBase : public UBTService
 {
-	GENERATED_BODY()
-public:
-	UBTService_BlackboardBase(const FObjectInitializer& ObjectInitializer = FObjectInitializer::Get());
+	GENERATED_UCLASS_BODY()
 
 	/** initialize any asset related data */
 	virtual void InitializeFromAsset(UBehaviorTree& Asset) override;
diff --git a/Engine/Source/Runtime/AIModule/Classes/BehaviorTree/Services/BTService_BlueprintBase.h b/Engine/Source/Runtime/AIModule/Classes/BehaviorTree/Services/BTService_BlueprintBase.h
--- a/Engine/Source/Runtime/AIModule/Classes/BehaviorTree/Services/BTService_BlueprintBase.h
+++ b/Engine/Source/Runtime/AIModule/Classes/BehaviorTree/Services/BTService_BlueprintBase.h
@@ -17,9 +17,7 @@ class FBehaviorBlueprintDetails;
 UCLASS(Abstract, Blueprintable)
 class AIMODULE_API UBTService_BlueprintBase : public UBTService
 {
-	GENERATED_BODY()
-public:
-	UBTService_BlueprintBase(const FObjectInitializer& ObjectInitializer = FObjectInitializer::Get());
+	GENERATED_UCLASS_BODY()
 
 	virtual void PostInitProperties() override;
 
diff --git a/Engine/Source/Runtime/AIModule/Classes/BehaviorTree/Tasks/BTTask_MakeNoise.h b/Engine/Source/Runtime/AIModule/Classes/BehaviorTree/Tasks/BTTask_MakeNoise.h
--- a/Engine/Source/Runtime/AIModule/Classes/BehaviorTree/Tasks/BTTask_MakeNoise.h
+++ b/Engine/Source/Runtime/AIModule/Classes/BehaviorTree/Tasks/BTTask_MakeNoise.h
@@ -10,9 +10,7 @@
 UCLASS()
 class AIMODULE_API UBTTask_MakeNoise : public UBTTaskNode
 {
-	GENERATED_BODY()
-public:
-	UBTTask_MakeNoise(const FObjectInitializer& ObjectInitializer = FObjectInitializer::Get());
+	GENERATED_UCLASS_BODY()
 
 	/** Loudnes of generated noise */
 	UPROPERTY(Category=Node, EditAnywhere, meta=(ClampMin = "0", ClampMax = "1", UIMin = "0", UIMax = "1"))
diff --git a/Engine/Source/Runtime/AIModule/Classes/BehaviorTree/Tasks/BTTask_MoveDirectlyToward.h b/Engine/Source/Runtime/AIModule/Classes/BehaviorTree/Tasks/BTTask_MoveDirectlyToward.h
--- a/Engine/Source/Runtime/AIModule/Classes/BehaviorTree/Tasks/BTTask_MoveDirectlyToward.h
+++ b/Engine/Source/Runtime/AIModule/Classes/BehaviorTree/Tasks/BTTask_MoveDirectlyToward.h
@@ -17,9 +17,7 @@ struct FBTMoveDirectlyTowardMemory
 UCLASS(config=Game)
 class AIMODULE_API UBTTask_MoveDirectlyToward : public UBTTask_BlackboardBase
 {
-	GENERATED_BODY()
-public:
-	UBTTask_MoveDirectlyToward(const FObjectInitializer& ObjectInitializer = FObjectInitializer::Get());
+	GENERATED_UCLASS_BODY()
 
 	UPROPERTY(config, Category=Node, EditAnywhere, meta=(ClampMin = "0.0"))
 	float AcceptableRadius;
diff --git a/Engine/Source/Runtime/AIModule/Classes/BehaviorTree/Tasks/BTTask_MoveTo.h b/Engine/Source/Runtime/AIModule/Classes/BehaviorTree/Tasks/BTTask_MoveTo.h
--- a/Engine/Source/Runtime/AIModule/Classes/BehaviorTree/Tasks/BTTask_MoveTo.h
+++ b/Engine/Source/Runtime/AIModule/Classes/BehaviorTree/Tasks/BTTask_MoveTo.h
@@ -21,9 +21,7 @@ struct FBTMoveToTaskMemory
 UCLASS(config=Game)
 class AIMODULE_API UBTTask_MoveTo : public UBTTask_BlackboardBase
 {
-	GENERATED_BODY()
-public:
-	UBTTask_MoveTo(const FObjectInitializer& ObjectInitializer = FObjectInitializer::Get());
+	GENERATED_UCLASS_BODY()
 
 	UPROPERTY(config, Category=Node, EditAnywhere, meta=(ClampMin = "0.0"))
 	float AcceptableRadius;
diff --git a/Engine/Source/Runtime/AIModule/Classes/BehaviorTree/Tasks/BTTask_PlayAnimation.h b/Engine/Source/Runtime/AIModule/Classes/BehaviorTree/Tasks/BTTask_PlayAnimation.h
--- a/Engine/Source/Runtime/AIModule/Classes/BehaviorTree/Tasks/BTTask_PlayAnimation.h
+++ b/Engine/Source/Runtime/AIModule/Classes/BehaviorTree/Tasks/BTTask_PlayAnimation.h
@@ -14,9 +14,7 @@ class UAnimationAsset;
 UCLASS()
 class AIMODULE_API UBTTask_PlayAnimation : public UBTTaskNode
 {
-	GENERATED_BODY()
-public:
-	UBTTask_PlayAnimation(const FObjectInitializer& ObjectInitializer = FObjectInitializer::Get());
+	GENERATED_UCLASS_BODY()
 
 	/** Animation asset to play. Note that it needs to match the skeleton of pawn this BT is controlling */
 	UPROPERTY(Category = Node, EditAnywhere)
diff --git a/Engine/Source/Runtime/AIModule/Classes/BehaviorTree/Tasks/BTTask_PlaySound.h b/Engine/Source/Runtime/AIModule/Classes/BehaviorTree/Tasks/BTTask_PlaySound.h
--- a/Engine/Source/Runtime/AIModule/Classes/BehaviorTree/Tasks/BTTask_PlaySound.h
+++ b/Engine/Source/Runtime/AIModule/Classes/BehaviorTree/Tasks/BTTask_PlaySound.h
@@ -13,9 +13,7 @@ class USoundCue;
 UCLASS()
 class AIMODULE_API UBTTask_PlaySound : public UBTTaskNode
 {
-	GENERATED_BODY()
-public:
-	UBTTask_PlaySound(const FObjectInitializer& ObjectInitializer = FObjectInitializer::Get());
+	GENERATED_UCLASS_BODY()
 
 	/** CUE to play */
 	UPROPERTY(Category=Node, EditAnywhere)
diff --git a/Engine/Source/Runtime/AIModule/Classes/BehaviorTree/Tasks/BTTask_PushPawnAction.h b/Engine/Source/Runtime/AIModule/Classes/BehaviorTree/Tasks/BTTask_PushPawnAction.h
--- a/Engine/Source/Runtime/AIModule/Classes/BehaviorTree/Tasks/BTTask_PushPawnAction.h
+++ b/Engine/Source/Runtime/AIModule/Classes/BehaviorTree/Tasks/BTTask_PushPawnAction.h
@@ -14,9 +14,7 @@ class UPawnAction;
 UCLASS()
 class AIMODULE_API UBTTask_PushPawnAction : public UBTTask_PawnActionBase
 {
-	GENERATED_BODY()
-public:
-	UBTTask_PushPawnAction(const FObjectInitializer& ObjectInitializer = FObjectInitializer::Get());
+	GENERATED_UCLASS_BODY()
 
 protected:
 	UPROPERTY(EditAnywhere, Instanced, Category = Action)
diff --git a/Engine/Source/Runtime/AIModule/Classes/BehaviorTree/Tasks/BTTask_RotateToFaceBBEntry.h b/Engine/Source/Runtime/AIModule/Classes/BehaviorTree/Tasks/BTTask_RotateToFaceBBEntry.h
--- a/Engine/Source/Runtime/AIModule/Classes/BehaviorTree/Tasks/BTTask_RotateToFaceBBEntry.h
+++ b/Engine/Source/Runtime/AIModule/Classes/BehaviorTree/Tasks/BTTask_RotateToFaceBBEntry.h
@@ -12,9 +12,7 @@
 UCLASS(config = Game)
 class AIMODULE_API UBTTask_RotateToFaceBBEntry : public UBTTask_BlackboardBase
 {
-	GENERATED_BODY()
-public:
-	UBTTask_RotateToFaceBBEntry(const FObjectInitializer& ObjectInitializer = FObjectInitializer::Get());
+	GENERATED_UCLASS_BODY()
 
 protected:
 	/** Success condition precision in decrees */
diff --git a/Engine/Source/Runtime/AIModule/Classes/BehaviorTree/Tasks/BTTask_RunBehavior.h b/Engine/Source/Runtime/AIModule/Classes/BehaviorTree/Tasks/BTTask_RunBehavior.h
--- a/Engine/Source/Runtime/AIModule/Classes/BehaviorTree/Tasks/BTTask_RunBehavior.h
+++ b/Engine/Source/Runtime/AIModule/Classes/BehaviorTree/Tasks/BTTask_RunBehavior.h
@@ -21,9 +21,7 @@
 UCLASS()
 class AIMODULE_API UBTTask_RunBehavior : public UBTTaskNode
 {
-	GENERATED_BODY()
-public:
-	UBTTask_RunBehavior(const FObjectInitializer& ObjectInitializer = FObjectInitializer::Get());
+	GENERATED_UCLASS_BODY()
 
 	virtual EBTNodeResult::Type ExecuteTask(UBehaviorTreeComponent& OwnerComp, uint8* NodeMemory) override;
 	virtual uint16 GetInstanceMemorySize() const override;
diff --git a/Engine/Source/Runtime/AIModule/Classes/BehaviorTree/Tasks/BTTask_RunEQSQuery.h b/Engine/Source/Runtime/AIModule/Classes/BehaviorTree/Tasks/BTTask_RunEQSQuery.h
--- a/Engine/Source/Runtime/AIModule/Classes/BehaviorTree/Tasks/BTTask_RunEQSQuery.h
+++ b/Engine/Source/Runtime/AIModule/Classes/BehaviorTree/Tasks/BTTask_RunEQSQuery.h
@@ -19,9 +19,7 @@ struct FBTEnvQueryTaskMemory
 UCLASS()
 class AIMODULE_API UBTTask_RunEQSQuery : public UBTTask_BlackboardBase
 {
-	GENERATED_BODY()
-public:
-	UBTTask_RunEQSQuery(const FObjectInitializer& ObjectInitializer = FObjectInitializer::Get());
+	GENERATED_UCLASS_BODY()
 
 	/** query to run */
 	UPROPERTY(Category=Node, EditAnywhere)
diff --git a/Engine/Source/Runtime/AIModule/Classes/BehaviorTree/Tasks/BTTask_SetTagCooldown.h b/Engine/Source/Runtime/AIModule/Classes/BehaviorTree/Tasks/BTTask_SetTagCooldown.h
--- a/Engine/Source/Runtime/AIModule/Classes/BehaviorTree/Tasks/BTTask_SetTagCooldown.h
+++ b/Engine/Source/Runtime/AIModule/Classes/BehaviorTree/Tasks/BTTask_SetTagCooldown.h
@@ -10,9 +10,7 @@
 UCLASS()
 class AIMODULE_API UBTTask_SetTagCooldown : public UBTTaskNode
 {
-	GENERATED_BODY()
-public:
-	UBTTask_SetTagCooldown(const FObjectInitializer& ObjectInitializer = FObjectInitializer::Get());
+	GENERATED_UCLASS_BODY()
 
 	/** gameplay tag that will be used for the cooldown */
 	UPROPERTY(Category = Cooldown, EditAnywhere)
diff --git a/Engine/Source/Runtime/AIModule/Classes/BehaviorTree/Tasks/BTTask_Wait.h b/Engine/Source/Runtime/AIModule/Classes/BehaviorTree/Tasks/BTTask_Wait.h
--- a/Engine/Source/Runtime/AIModule/Classes/BehaviorTree/Tasks/BTTask_Wait.h
+++ b/Engine/Source/Runtime/AIModule/Classes/BehaviorTree/Tasks/BTTask_Wait.h
@@ -16,9 +16,7 @@ struct FBTWaitTaskMemory
 UCLASS()
 class AIMODULE_API UBTTask_Wait : public UBTTaskNode
 {
-	GENERATED_BODY()
-public:
-	UBTTask_Wait(const FObjectInitializer& ObjectInitializer = FObjectInitializer::Get());
+	GENERATED_UCLASS_BODY()
 
 	/** wait time in seconds */
 	UPROPERTY(Category = Wait, EditAnywhere, meta = (ClampMin = "0.0", UIMin = "0.0"))
diff --git a/Engine/Source/Runtime/AIModule/Classes/BehaviorTree/Tasks/BTTask_WaitBlackboardTime.h b/Engine/Source/Runtime/AIModule/Classes/BehaviorTree/Tasks/BTTask_WaitBlackboardTime.h
--- a/Engine/Source/Runtime/AIModule/Classes/BehaviorTree/Tasks/BTTask_WaitBlackboardTime.h
+++ b/Engine/Source/Runtime/AIModule/Classes/BehaviorTree/Tasks/BTTask_WaitBlackboardTime.h
@@ -11,9 +11,7 @@
 UCLASS(hidecategories=Wait)
 class AIMODULE_API UBTTask_WaitBlackboardTime : public UBTTask_Wait
 {
-	GENERATED_BODY()
-public:
-	UBTTask_WaitBlackboardTime(const FObjectInitializer& ObjectInitializer = FObjectInitializer::Get());
+	GENERATED_UCLASS_BODY()
 
 	virtual void InitializeFromAsset(UBehaviorTree& Asset) override;
 	virtual EBTNodeResult::Type ExecuteTask(UBehaviorTreeComponent& OwnerComp, uint8* NodeMemory) override;
diff --git a/Engine/Source/Runtime/AIModule/Classes/Blueprint/AIBlueprintHelperLibrary.h b/Engine/Source/Runtime/AIModule/Classes/Blueprint/AIBlueprintHelperLibrary.h
--- a/Engine/Source/Runtime/AIModule/Classes/Blueprint/AIBlueprintHelperLibrary.h
+++ b/Engine/Source/Runtime/AIModule/Classes/Blueprint/AIBlueprintHelperLibrary.h
@@ -18,9 +18,7 @@ class UAIAsyncTaskBlueprintProxy;
 UCLASS()
 class AIMODULE_API UAIBlueprintHelperLibrary : public UBlueprintFunctionLibrary
 {
-	GENERATED_BODY()
-public:
-	UAIBlueprintHelperLibrary(const FObjectInitializer& ObjectInitializer = FObjectInitializer::Get());
+	GENERATED_UCLASS_BODY()
 
 	UFUNCTION(BlueprintCallable, meta=(WorldContext="WorldContextObject", BlueprintInternalUseOnly = "TRUE"))
 	static UAIAsyncTaskBlueprintProxy* CreateMoveToProxyObject(UObject* WorldContextObject, APawn* Pawn, FVector Destination, AActor* TargetActor = NULL, float AcceptanceRadius = 5.f, bool bStopOnOverlap = false);
diff --git a/Engine/Source/Runtime/AIModule/Classes/BrainComponent.h b/Engine/Source/Runtime/AIModule/Classes/BrainComponent.h
--- a/Engine/Source/Runtime/AIModule/Classes/BrainComponent.h
+++ b/Engine/Source/Runtime/AIModule/Classes/BrainComponent.h
@@ -105,9 +105,7 @@ private:
 UCLASS(abstract, BlueprintType)
 class AIMODULE_API UBrainComponent : public UActorComponent, public IAIResourceInterface
 {
-	GENERATED_BODY()
-public:
-	UBrainComponent(const FObjectInitializer& ObjectInitializer = FObjectInitializer::Get());
+	GENERATED_UCLASS_BODY()
 
 protected:
 	/** blackboard component */
@@ -211,4 +209,4 @@ FORCEINLINE UBlackboardComponent* UBrainComponent::GetBlackboardComponent()
 FORCEINLINE const UBlackboardComponent* UBrainComponent::GetBlackboardComponent() const
 {
 	return BlackboardComp;
-}
+}
\ No newline at end of file
diff --git a/Engine/Source/Runtime/AIModule/Classes/DataProviders/AIDataProvider.h b/Engine/Source/Runtime/AIModule/Classes/DataProviders/AIDataProvider.h
--- a/Engine/Source/Runtime/AIModule/Classes/DataProviders/AIDataProvider.h
+++ b/Engine/Source/Runtime/AIModule/Classes/DataProviders/AIDataProvider.h
@@ -20,7 +20,7 @@ class UAIDataProvider;
 USTRUCT()
 struct AIMODULE_API FAIDataProviderValue
 {
-	GENERATED_BODY()
+	GENERATED_USTRUCT_BODY()
 
 	/** cached uproperty of provider */
 	UPROPERTY(transient)
@@ -57,7 +57,7 @@ struct AIMODULE_API FAIDataProviderValue
 USTRUCT()
 struct AIMODULE_API FAIDataProviderTypedValue : public FAIDataProviderValue
 {
-	GENERATED_BODY()
+	GENERATED_USTRUCT_BODY()
 
 	/** type of value */
 	UPROPERTY()
@@ -70,7 +70,7 @@ struct AIMODULE_API FAIDataProviderTypedValue : public FAIDataProviderValue
 USTRUCT()
 struct AIMODULE_API FAIDataProviderStructValue : public FAIDataProviderValue
 {
-	GENERATED_BODY()
+	GENERATED_USTRUCT_BODY()
 
 	/** name of UStruct type */
 	FString StructName;
@@ -81,7 +81,7 @@ struct AIMODULE_API FAIDataProviderStructValue : public FAIDataProviderValue
 USTRUCT()
 struct AIMODULE_API FAIDataProviderIntValue : public FAIDataProviderTypedValue
 {
-	GENERATED_BODY()
+	GENERATED_USTRUCT_BODY()
 	FAIDataProviderIntValue();
 
 	UPROPERTY(EditAnywhere, Category = Value)
@@ -94,7 +94,7 @@ struct AIMODULE_API FAIDataProviderIntValue : public FAIDataProviderTypedValue
 USTRUCT()
 struct AIMODULE_API FAIDataProviderFloatValue : public FAIDataProviderTypedValue
 {
-	GENERATED_BODY()
+	GENERATED_USTRUCT_BODY()
 	FAIDataProviderFloatValue();
 
 	UPROPERTY(EditAnywhere, Category = Value)
@@ -107,7 +107,7 @@ struct AIMODULE_API FAIDataProviderFloatValue : public FAIDataProviderTypedValue
 USTRUCT()
 struct AIMODULE_API FAIDataProviderBoolValue : public FAIDataProviderTypedValue
 {
-	GENERATED_BODY()
+	GENERATED_USTRUCT_BODY()
 	FAIDataProviderBoolValue();
 
 	UPROPERTY(EditAnywhere, Category = Value)
@@ -120,9 +120,7 @@ struct AIMODULE_API FAIDataProviderBoolValue : public FAIDataProviderTypedValue
 UCLASS(EditInlineNew, Abstract, CollapseCategories, AutoExpandCategories=(Provider))
 class AIMODULE_API UAIDataProvider : public UObject
 {
-	GENERATED_BODY()
-public:
-	UAIDataProvider(const FObjectInitializer& ObjectInitializer = FObjectInitializer::Get());
+	GENERATED_UCLASS_BODY()
 
 	virtual void BindData(UObject* Owner, int32 RequestId);
 	virtual FString ToString(FName PropName) const;
diff --git a/Engine/Source/Runtime/AIModule/Classes/DataProviders/AIDataProvider_QueryParams.h b/Engine/Source/Runtime/AIModule/Classes/DataProviders/AIDataProvider_QueryParams.h
--- a/Engine/Source/Runtime/AIModule/Classes/DataProviders/AIDataProvider_QueryParams.h
+++ b/Engine/Source/Runtime/AIModule/Classes/DataProviders/AIDataProvider_QueryParams.h
@@ -16,9 +16,7 @@
 UCLASS(EditInlineNew, meta=(DisplayName="Query Params"))
 class AIMODULE_API UAIDataProvider_QueryParams : public UAIDataProvider
 {
-	GENERATED_BODY()
-public:
-	UAIDataProvider_QueryParams(const FObjectInitializer& ObjectInitializer = FObjectInitializer::Get());
+	GENERATED_UCLASS_BODY()
 
 	virtual void BindData(UObject* Owner, int32 RequestId);
 	virtual FString ToString(FName PropName) const override;
diff --git a/Engine/Source/Runtime/AIModule/Classes/EnvironmentQuery/Contexts/EnvQueryContext_BlueprintBase.h b/Engine/Source/Runtime/AIModule/Classes/EnvironmentQuery/Contexts/EnvQueryContext_BlueprintBase.h
--- a/Engine/Source/Runtime/AIModule/Classes/EnvironmentQuery/Contexts/EnvQueryContext_BlueprintBase.h
+++ b/Engine/Source/Runtime/AIModule/Classes/EnvironmentQuery/Contexts/EnvQueryContext_BlueprintBase.h
@@ -12,9 +12,7 @@ struct FEnvQueryContextData;
 UCLASS(MinimalAPI, Abstract, Blueprintable)
 class UEnvQueryContext_BlueprintBase : public UEnvQueryContext
 {
-	GENERATED_BODY()
-public:
-	AIMODULE_API UEnvQueryContext_BlueprintBase(const FObjectInitializer& ObjectInitializer = FObjectInitializer::Get());
+	GENERATED_UCLASS_BODY()
 
 	enum ECallMode
 	{
diff --git a/Engine/Source/Runtime/AIModule/Classes/EnvironmentQuery/Contexts/EnvQueryContext_Item.h b/Engine/Source/Runtime/AIModule/Classes/EnvironmentQuery/Contexts/EnvQueryContext_Item.h
--- a/Engine/Source/Runtime/AIModule/Classes/EnvironmentQuery/Contexts/EnvQueryContext_Item.h
+++ b/Engine/Source/Runtime/AIModule/Classes/EnvironmentQuery/Contexts/EnvQueryContext_Item.h
@@ -7,7 +7,5 @@
 UCLASS()
 class AIMODULE_API UEnvQueryContext_Item : public UEnvQueryContext
 {
-	GENERATED_BODY()
-public:
-	UEnvQueryContext_Item(const FObjectInitializer& ObjectInitializer = FObjectInitializer::Get());
+	GENERATED_UCLASS_BODY()
 };
diff --git a/Engine/Source/Runtime/AIModule/Classes/EnvironmentQuery/Contexts/EnvQueryContext_Querier.h b/Engine/Source/Runtime/AIModule/Classes/EnvironmentQuery/Contexts/EnvQueryContext_Querier.h
--- a/Engine/Source/Runtime/AIModule/Classes/EnvironmentQuery/Contexts/EnvQueryContext_Querier.h
+++ b/Engine/Source/Runtime/AIModule/Classes/EnvironmentQuery/Contexts/EnvQueryContext_Querier.h
@@ -10,9 +10,7 @@ struct FEnvQueryContextData;
 UCLASS(MinimalAPI)
 class UEnvQueryContext_Querier : public UEnvQueryContext
 {
-	GENERATED_BODY()
-public:
-	AIMODULE_API UEnvQueryContext_Querier(const FObjectInitializer& ObjectInitializer = FObjectInitializer::Get());
+	GENERATED_UCLASS_BODY()
  
 	virtual void ProvideContext(FEnvQueryInstance& QueryInstance, FEnvQueryContextData& ContextData) const override;
 };
diff --git a/Engine/Source/Runtime/AIModule/Classes/EnvironmentQuery/EQSRenderingComponent.h b/Engine/Source/Runtime/AIModule/Classes/EnvironmentQuery/EQSRenderingComponent.h
--- a/Engine/Source/Runtime/AIModule/Classes/EnvironmentQuery/EQSRenderingComponent.h
+++ b/Engine/Source/Runtime/AIModule/Classes/EnvironmentQuery/EQSRenderingComponent.h
@@ -39,9 +39,7 @@ private:
 UCLASS(hidecategories=Object)
 class AIMODULE_API UEQSRenderingComponent : public UPrimitiveComponent
 {
-	GENERATED_BODY()
-public:
-	UEQSRenderingComponent(const FObjectInitializer& ObjectInitializer = FObjectInitializer::Get());
+	GENERATED_UCLASS_BODY()
 
 	FString DrawFlagName;
 	uint32 bDrawOnlyWhenSelected : 1;
diff --git a/Engine/Source/Runtime/AIModule/Classes/EnvironmentQuery/EQSTestingPawn.h b/Engine/Source/Runtime/AIModule/Classes/EnvironmentQuery/EQSTestingPawn.h
--- a/Engine/Source/Runtime/AIModule/Classes/EnvironmentQuery/EQSTestingPawn.h
+++ b/Engine/Source/Runtime/AIModule/Classes/EnvironmentQuery/EQSTestingPawn.h
@@ -24,9 +24,7 @@ enum class EEnvQueryHightlightMode : uint8
 UCLASS(abstract, hidecategories=(Advanced, Attachment, Collision, Mesh, Animation, Clothing, Physics, Rendering, Lighting, Activation, CharacterMovement, AgentPhysics, Avoidance, MovementComponent, Velocity, Shape, Camera, Input, Layers, SkeletalMesh, Optimization, Pawn, Replication, Actor))
 class AIMODULE_API AEQSTestingPawn : public ACharacter, public IEQSQueryResultSourceInterface
 {
-	GENERATED_BODY()
-public:
-	AEQSTestingPawn(const FObjectInitializer& ObjectInitializer = FObjectInitializer::Get());
+	GENERATED_UCLASS_BODY()
 	
 	UPROPERTY(Category=EQS, EditAnywhere)
 	UEnvQuery* QueryTemplate;
diff --git a/Engine/Source/Runtime/AIModule/Classes/EnvironmentQuery/EnvQueryContext.h b/Engine/Source/Runtime/AIModule/Classes/EnvironmentQuery/EnvQueryContext.h
--- a/Engine/Source/Runtime/AIModule/Classes/EnvironmentQuery/EnvQueryContext.h
+++ b/Engine/Source/Runtime/AIModule/Classes/EnvironmentQuery/EnvQueryContext.h
@@ -9,9 +9,7 @@ struct FEnvQueryInstance;
 UCLASS(Abstract)
 class AIMODULE_API UEnvQueryContext : public UObject
 {
-	GENERATED_BODY()
-public:
-	UEnvQueryContext(const FObjectInitializer& ObjectInitializer = FObjectInitializer::Get());
+	GENERATED_UCLASS_BODY()
 
 	virtual void ProvideContext(FEnvQueryInstance& QueryInstance, FEnvQueryContextData& ContextData) const;
 };
diff --git a/Engine/Source/Runtime/AIModule/Classes/EnvironmentQuery/EnvQueryManager.h b/Engine/Source/Runtime/AIModule/Classes/EnvironmentQuery/EnvQueryManager.h
--- a/Engine/Source/Runtime/AIModule/Classes/EnvironmentQuery/EnvQueryManager.h
+++ b/Engine/Source/Runtime/AIModule/Classes/EnvironmentQuery/EnvQueryManager.h
@@ -18,7 +18,7 @@ class UEnvQueryTest;
 USTRUCT()
 struct AIMODULE_API FEnvQueryRequest
 {
-	GENERATED_BODY()
+	GENERATED_USTRUCT_BODY()
 
 	FEnvQueryRequest() : QueryTemplate(NULL), Owner(NULL), World(NULL) {}
 
@@ -74,7 +74,7 @@ protected:
 USTRUCT()
 struct FEnvQueryInstanceCache
 {
-	GENERATED_BODY()
+	GENERATED_USTRUCT_BODY()
 
 	/** query template, duplicated in manager's world */
 	UPROPERTY()
@@ -109,9 +109,7 @@ FORCEINLINE bool operator== (const FEQSDebugger::FEnvQueryInfo & Left, const FEQ
 UCLASS()
 class AIMODULE_API UEnvQueryManager : public UObject, public FTickableGameObject
 {
-	GENERATED_BODY()
-public:
-	UEnvQueryManager(const FObjectInitializer& ObjectInitializer = FObjectInitializer::Get());
+	GENERATED_UCLASS_BODY()
 
 	// We need to implement GetWorld() so that any EQS-related blueprints (such as blueprint contexts) can implement
 	// GetWorld() and so provide access to blueprint nodes using hidden WorldContextObject parameters.
diff --git a/Engine/Source/Runtime/AIModule/Classes/EnvironmentQuery/EnvQueryNode.h b/Engine/Source/Runtime/AIModule/Classes/EnvironmentQuery/EnvQueryNode.h
--- a/Engine/Source/Runtime/AIModule/Classes/EnvironmentQuery/EnvQueryNode.h
+++ b/Engine/Source/Runtime/AIModule/Classes/EnvironmentQuery/EnvQueryNode.h
@@ -6,9 +6,7 @@
 UCLASS(Abstract)
 class AIMODULE_API UEnvQueryNode : public UObject
 {
-	GENERATED_BODY()
-public:
-	UEnvQueryNode(const FObjectInitializer& ObjectInitializer = FObjectInitializer::Get());
+	GENERATED_UCLASS_BODY()
 
 	/** Versioning for updating deprecated properties */
 	UPROPERTY()
diff --git a/Engine/Source/Runtime/AIModule/Classes/EnvironmentQuery/EnvQueryOption.h b/Engine/Source/Runtime/AIModule/Classes/EnvironmentQuery/EnvQueryOption.h
--- a/Engine/Source/Runtime/AIModule/Classes/EnvironmentQuery/EnvQueryOption.h
+++ b/Engine/Source/Runtime/AIModule/Classes/EnvironmentQuery/EnvQueryOption.h
@@ -9,9 +9,7 @@ class UEnvQueryTest;
 UCLASS()
 class AIMODULE_API UEnvQueryOption : public UObject
 {
-	GENERATED_BODY()
-public:
-	UEnvQueryOption(const FObjectInitializer& ObjectInitializer = FObjectInitializer::Get());
+	GENERATED_UCLASS_BODY()
 
 	UPROPERTY()
 	UEnvQueryGenerator* Generator;
diff --git a/Engine/Source/Runtime/AIModule/Classes/EnvironmentQuery/EnvQueryTest.h b/Engine/Source/Runtime/AIModule/Classes/EnvironmentQuery/EnvQueryTest.h
--- a/Engine/Source/Runtime/AIModule/Classes/EnvironmentQuery/EnvQueryTest.h
+++ b/Engine/Source/Runtime/AIModule/Classes/EnvironmentQuery/EnvQueryTest.h
@@ -38,9 +38,7 @@ struct FEnvQueryTestScoringPreview
 UCLASS(Abstract)
 class AIMODULE_API UEnvQueryTest : public UEnvQueryNode
 {
-	GENERATED_BODY()
-public:
-	UEnvQueryTest(const FObjectInitializer& ObjectInitializer = FObjectInitializer::Get());
+	GENERATED_UCLASS_BODY()
 
 	/** Number of test as defined in data asset */
 	UPROPERTY()
diff --git a/Engine/Source/Runtime/AIModule/Classes/EnvironmentQuery/EnvQueryTypes.h b/Engine/Source/Runtime/AIModule/Classes/EnvironmentQuery/EnvQueryTypes.h
--- a/Engine/Source/Runtime/AIModule/Classes/EnvironmentQuery/EnvQueryTypes.h
+++ b/Engine/Source/Runtime/AIModule/Classes/EnvironmentQuery/EnvQueryTypes.h
@@ -190,7 +190,7 @@ namespace EEnvQueryTestClamping
 USTRUCT(BlueprintType)
 struct AIMODULE_API FEnvNamedValue
 {
-	GENERATED_BODY()
+	GENERATED_USTRUCT_BODY();
 
 	UPROPERTY(EditAnywhere, BlueprintReadWrite, Category=Param)
 	FName ParamName;
@@ -205,7 +205,7 @@ struct AIMODULE_API FEnvNamedValue
 USTRUCT()
 struct AIMODULE_API FEnvDirection
 {
-	GENERATED_BODY()
+	GENERATED_USTRUCT_BODY()
 
 	/** line A: start context */
 	UPROPERTY(EditDefaultsOnly, Category=Direction)
@@ -229,7 +229,7 @@ struct AIMODULE_API FEnvDirection
 USTRUCT()
 struct AIMODULE_API FEnvTraceData
 {
-	GENERATED_BODY()
+	GENERATED_USTRUCT_BODY()
 
 	enum EDescriptionMode
 	{
@@ -886,7 +886,7 @@ namespace FEQSHelpers
 USTRUCT()
 struct AIMODULE_API FEnvFloatParam_DEPRECATED
 {
-	GENERATED_BODY();
+	GENERATED_USTRUCT_BODY();
 
 	/** default value */
 	UPROPERTY(EditDefaultsOnly, Category = Param)
@@ -903,7 +903,7 @@ struct AIMODULE_API FEnvFloatParam_DEPRECATED
 USTRUCT()
 struct AIMODULE_API FEnvIntParam_DEPRECATED
 {
-	GENERATED_BODY();
+	GENERATED_USTRUCT_BODY();
 
 	/** default value */
 	UPROPERTY(EditDefaultsOnly, Category = Param)
@@ -920,7 +920,7 @@ struct AIMODULE_API FEnvIntParam_DEPRECATED
 USTRUCT()
 struct AIMODULE_API FEnvBoolParam_DEPRECATED
 {
-	GENERATED_BODY();
+	GENERATED_USTRUCT_BODY();
 
 	/** default value */
 	UPROPERTY(EditDefaultsOnly, Category = Param)
@@ -937,19 +937,19 @@ struct AIMODULE_API FEnvBoolParam_DEPRECATED
 USTRUCT()
 struct DEPRECATED(4.8, "FEnvFloatParam is deprecated in 4.8 and was replaced with FAIDataProviderFloatValue. Please use that type instead.") AIMODULE_API FEnvFloatParam : public FEnvFloatParam_DEPRECATED
 {
-	GENERATED_BODY();
+	GENERATED_USTRUCT_BODY();
 };
 
 USTRUCT()
 struct DEPRECATED(4.8, "FEnvIntParam is deprecated in 4.8 and was replaced with FAIDataProviderIntValue. Please use that type instead.") AIMODULE_API FEnvIntParam : public FEnvIntParam_DEPRECATED
 {
-	GENERATED_BODY();
+	GENERATED_USTRUCT_BODY();
 };
 
 USTRUCT()
 struct DEPRECATED(4.8, "FEnvBoolParam is deprecated in 4.8 and was replaced with FAIDataProviderBoolValue. Please use that type instead.") AIMODULE_API FEnvBoolParam : public FEnvBoolParam_DEPRECATED
 {
-	GENERATED_BODY();
+	GENERATED_USTRUCT_BODY();
 };
 
 // END DEPRECATED SUPPORT
diff --git a/Engine/Source/Runtime/AIModule/Classes/EnvironmentQuery/Generators/EnvQueryGenerator_ActorsOfClass.h b/Engine/Source/Runtime/AIModule/Classes/EnvironmentQuery/Generators/EnvQueryGenerator_ActorsOfClass.h
--- a/Engine/Source/Runtime/AIModule/Classes/EnvironmentQuery/Generators/EnvQueryGenerator_ActorsOfClass.h
+++ b/Engine/Source/Runtime/AIModule/Classes/EnvironmentQuery/Generators/EnvQueryGenerator_ActorsOfClass.h
@@ -10,9 +10,7 @@ class UEnvQueryContext;
 UCLASS(meta = (DisplayName = "Actors Of Class"))
 class UEnvQueryGenerator_ActorsOfClass : public UEnvQueryGenerator
 {
-	GENERATED_BODY()
-public:
-	UEnvQueryGenerator_ActorsOfClass(const FObjectInitializer& ObjectInitializer = FObjectInitializer::Get());
+	GENERATED_UCLASS_BODY()
 
 	/** max distance of path between point and context */
 	UPROPERTY(EditDefaultsOnly, Category=Generator)
diff --git a/Engine/Source/Runtime/AIModule/Classes/EnvironmentQuery/Generators/EnvQueryGenerator_BlueprintBase.h b/Engine/Source/Runtime/AIModule/Classes/EnvironmentQuery/Generators/EnvQueryGenerator_BlueprintBase.h
--- a/Engine/Source/Runtime/AIModule/Classes/EnvironmentQuery/Generators/EnvQueryGenerator_BlueprintBase.h
+++ b/Engine/Source/Runtime/AIModule/Classes/EnvironmentQuery/Generators/EnvQueryGenerator_BlueprintBase.h
@@ -12,9 +12,7 @@ class UEnvQueryItemType;
 UCLASS(Abstract, Blueprintable)
 class AIMODULE_API UEnvQueryGenerator_BlueprintBase : public UEnvQueryGenerator
 {
-	GENERATED_BODY()
-public:
-	UEnvQueryGenerator_BlueprintBase(const FObjectInitializer& ObjectInitializer = FObjectInitializer::Get());
+	GENERATED_UCLASS_BODY()
 
 	/** A short description of what test does, like "Generate pawn named Joe" */
 	UPROPERTY(EditAnywhere, Category = Generator)
diff --git a/Engine/Source/Runtime/AIModule/Classes/EnvironmentQuery/Generators/EnvQueryGenerator_Donut.h b/Engine/Source/Runtime/AIModule/Classes/EnvironmentQuery/Generators/EnvQueryGenerator_Donut.h
--- a/Engine/Source/Runtime/AIModule/Classes/EnvironmentQuery/Generators/EnvQueryGenerator_Donut.h
+++ b/Engine/Source/Runtime/AIModule/Classes/EnvironmentQuery/Generators/EnvQueryGenerator_Donut.h
@@ -9,9 +9,7 @@ struct FEnvQueryInstance;
 UCLASS(meta = (DisplayName = "Points: Donut"))
 class AIMODULE_API UEnvQueryGenerator_Donut : public UEnvQueryGenerator_ProjectedPoints
 {
-	GENERATED_BODY()
-public:
-	UEnvQueryGenerator_Donut(const FObjectInitializer& ObjectInitializer = FObjectInitializer::Get());
+	GENERATED_UCLASS_BODY()
 
 	/** min distance between point and context */
 	UPROPERTY(EditDefaultsOnly, Category = Generator)
diff --git a/Engine/Source/Runtime/AIModule/Classes/EnvironmentQuery/Generators/EnvQueryGenerator_OnCircle.h b/Engine/Source/Runtime/AIModule/Classes/EnvironmentQuery/Generators/EnvQueryGenerator_OnCircle.h
--- a/Engine/Source/Runtime/AIModule/Classes/EnvironmentQuery/Generators/EnvQueryGenerator_OnCircle.h
+++ b/Engine/Source/Runtime/AIModule/Classes/EnvironmentQuery/Generators/EnvQueryGenerator_OnCircle.h
@@ -9,9 +9,7 @@ struct FEnvQueryInstance;
 UCLASS(meta = (DisplayName = "Points: Circle"))
 class AIMODULE_API UEnvQueryGenerator_OnCircle : public UEnvQueryGenerator_ProjectedPoints
 {
-	GENERATED_BODY()
-public:
-	UEnvQueryGenerator_OnCircle(const FObjectInitializer& ObjectInitializer = FObjectInitializer::Get());
+	GENERATED_UCLASS_BODY()
 
 	/** max distance of path between point and context */
 	UPROPERTY(EditDefaultsOnly, Category=Generator)
diff --git a/Engine/Source/Runtime/AIModule/Classes/EnvironmentQuery/Generators/EnvQueryGenerator_PathingGrid.h b/Engine/Source/Runtime/AIModule/Classes/EnvironmentQuery/Generators/EnvQueryGenerator_PathingGrid.h
--- a/Engine/Source/Runtime/AIModule/Classes/EnvironmentQuery/Generators/EnvQueryGenerator_PathingGrid.h
+++ b/Engine/Source/Runtime/AIModule/Classes/EnvironmentQuery/Generators/EnvQueryGenerator_PathingGrid.h
@@ -14,9 +14,7 @@ class UNavigationQueryFilter;
 UCLASS(meta = (DisplayName = "Points: Pathing Grid"))
 class UEnvQueryGenerator_PathingGrid : public UEnvQueryGenerator_SimpleGrid
 {
-	GENERATED_BODY()
-public:
-	UEnvQueryGenerator_PathingGrid(const FObjectInitializer& ObjectInitializer = FObjectInitializer::Get());
+	GENERATED_UCLASS_BODY()
 
 	/** pathfinding direction */
 	UPROPERTY(EditDefaultsOnly, Category = Generator)
diff --git a/Engine/Source/Runtime/AIModule/Classes/EnvironmentQuery/Generators/EnvQueryGenerator_ProjectedPoints.h b/Engine/Source/Runtime/AIModule/Classes/EnvironmentQuery/Generators/EnvQueryGenerator_ProjectedPoints.h
--- a/Engine/Source/Runtime/AIModule/Classes/EnvironmentQuery/Generators/EnvQueryGenerator_ProjectedPoints.h
+++ b/Engine/Source/Runtime/AIModule/Classes/EnvironmentQuery/Generators/EnvQueryGenerator_ProjectedPoints.h
@@ -11,9 +11,7 @@ class ANavigationData;
 UCLASS(Abstract)
 class AIMODULE_API UEnvQueryGenerator_ProjectedPoints : public UEnvQueryGenerator
 {
-	GENERATED_BODY()
-public:
-	UEnvQueryGenerator_ProjectedPoints(const FObjectInitializer& ObjectInitializer = FObjectInitializer::Get());
+	GENERATED_UCLASS_BODY()
 
 	/** trace params */
 	UPROPERTY(EditDefaultsOnly, Category = Generator)
diff --git a/Engine/Source/Runtime/AIModule/Classes/EnvironmentQuery/Generators/EnvQueryGenerator_SimpleGrid.h b/Engine/Source/Runtime/AIModule/Classes/EnvironmentQuery/Generators/EnvQueryGenerator_SimpleGrid.h
--- a/Engine/Source/Runtime/AIModule/Classes/EnvironmentQuery/Generators/EnvQueryGenerator_SimpleGrid.h
+++ b/Engine/Source/Runtime/AIModule/Classes/EnvironmentQuery/Generators/EnvQueryGenerator_SimpleGrid.h
@@ -13,9 +13,7 @@ class UEnvQueryContext;
 UCLASS(meta = (DisplayName = "Points: Grid"))
 class UEnvQueryGenerator_SimpleGrid : public UEnvQueryGenerator_ProjectedPoints
 {
-	GENERATED_BODY()
-public:
-	UEnvQueryGenerator_SimpleGrid(const FObjectInitializer& ObjectInitializer = FObjectInitializer::Get());
+	GENERATED_UCLASS_BODY()
 
 	/** square's extent */
 	UPROPERTY(EditDefaultsOnly, Category=Generator)
diff --git a/Engine/Source/Runtime/AIModule/Classes/EnvironmentQuery/Items/EnvQueryItemType.h b/Engine/Source/Runtime/AIModule/Classes/EnvironmentQuery/Items/EnvQueryItemType.h
--- a/Engine/Source/Runtime/AIModule/Classes/EnvironmentQuery/Items/EnvQueryItemType.h
+++ b/Engine/Source/Runtime/AIModule/Classes/EnvironmentQuery/Items/EnvQueryItemType.h
@@ -9,9 +9,7 @@ struct FBlackboardKeySelector;
 UCLASS(Abstract)
 class AIMODULE_API UEnvQueryItemType : public UObject
 {
-	GENERATED_BODY()
-public:
-	UEnvQueryItemType(const FObjectInitializer& ObjectInitializer = FObjectInitializer::Get());
+	GENERATED_UCLASS_BODY()
 
 	/** get ValueSize */
 	FORCEINLINE uint16 GetValueSize() const { return ValueSize; }
diff --git a/Engine/Source/Runtime/AIModule/Classes/EnvironmentQuery/Items/EnvQueryItemType_Actor.h b/Engine/Source/Runtime/AIModule/Classes/EnvironmentQuery/Items/EnvQueryItemType_Actor.h
--- a/Engine/Source/Runtime/AIModule/Classes/EnvironmentQuery/Items/EnvQueryItemType_Actor.h
+++ b/Engine/Source/Runtime/AIModule/Classes/EnvironmentQuery/Items/EnvQueryItemType_Actor.h
@@ -10,9 +10,7 @@ struct FEnvQueryContextData;
 UCLASS()
 class AIMODULE_API UEnvQueryItemType_Actor : public UEnvQueryItemType_ActorBase
 {
-	GENERATED_BODY()
-public:
-	UEnvQueryItemType_Actor(const FObjectInitializer& ObjectInitializer = FObjectInitializer::Get());
+	GENERATED_UCLASS_BODY()
 
 	static AActor* GetValue(const uint8* RawData);
 	static void SetValue(uint8* RawData, const AActor* Value);
diff --git a/Engine/Source/Runtime/AIModule/Classes/EnvironmentQuery/Items/EnvQueryItemType_ActorBase.h b/Engine/Source/Runtime/AIModule/Classes/EnvironmentQuery/Items/EnvQueryItemType_ActorBase.h
--- a/Engine/Source/Runtime/AIModule/Classes/EnvironmentQuery/Items/EnvQueryItemType_ActorBase.h
+++ b/Engine/Source/Runtime/AIModule/Classes/EnvironmentQuery/Items/EnvQueryItemType_ActorBase.h
@@ -11,9 +11,7 @@ class UBlackboardComponent;
 UCLASS(Abstract)
 class AIMODULE_API UEnvQueryItemType_ActorBase : public UEnvQueryItemType_VectorBase
 {
-	GENERATED_BODY()
-public:
-	UEnvQueryItemType_ActorBase(const FObjectInitializer& ObjectInitializer = FObjectInitializer::Get());
+	GENERATED_UCLASS_BODY()
 
 	virtual AActor* GetActor(const uint8* RawData) const;
 
diff --git a/Engine/Source/Runtime/AIModule/Classes/EnvironmentQuery/Items/EnvQueryItemType_Point.h b/Engine/Source/Runtime/AIModule/Classes/EnvironmentQuery/Items/EnvQueryItemType_Point.h
--- a/Engine/Source/Runtime/AIModule/Classes/EnvironmentQuery/Items/EnvQueryItemType_Point.h
+++ b/Engine/Source/Runtime/AIModule/Classes/EnvironmentQuery/Items/EnvQueryItemType_Point.h
@@ -9,9 +9,7 @@ struct FEnvQueryContextData;
 UCLASS()
 class AIMODULE_API UEnvQueryItemType_Point : public UEnvQueryItemType_VectorBase
 {
-	GENERATED_BODY()
-public:
-	UEnvQueryItemType_Point(const FObjectInitializer& ObjectInitializer = FObjectInitializer::Get());
+	GENERATED_UCLASS_BODY()
 
 	static FVector GetValue(const uint8* RawData);
 	static void SetValue(uint8* RawData, const FVector& Value);
diff --git a/Engine/Source/Runtime/AIModule/Classes/EnvironmentQuery/Items/EnvQueryItemType_VectorBase.h b/Engine/Source/Runtime/AIModule/Classes/EnvironmentQuery/Items/EnvQueryItemType_VectorBase.h
--- a/Engine/Source/Runtime/AIModule/Classes/EnvironmentQuery/Items/EnvQueryItemType_VectorBase.h
+++ b/Engine/Source/Runtime/AIModule/Classes/EnvironmentQuery/Items/EnvQueryItemType_VectorBase.h
@@ -10,9 +10,7 @@ class UBlackboardComponent;
 UCLASS(Abstract)
 class AIMODULE_API UEnvQueryItemType_VectorBase : public UEnvQueryItemType
 {
-	GENERATED_BODY()
-public:
-	UEnvQueryItemType_VectorBase(const FObjectInitializer& ObjectInitializer = FObjectInitializer::Get());
+	GENERATED_UCLASS_BODY()
 
 	virtual FVector GetItemLocation(const uint8* RawData) const;
 	virtual FRotator GetItemRotation(const uint8* RawData) const;
diff --git a/Engine/Source/Runtime/AIModule/Classes/EnvironmentQuery/Tests/EnvQueryTest_Distance.h b/Engine/Source/Runtime/AIModule/Classes/EnvironmentQuery/Tests/EnvQueryTest_Distance.h
--- a/Engine/Source/Runtime/AIModule/Classes/EnvironmentQuery/Tests/EnvQueryTest_Distance.h
+++ b/Engine/Source/Runtime/AIModule/Classes/EnvironmentQuery/Tests/EnvQueryTest_Distance.h
@@ -20,9 +20,7 @@ namespace EEnvTestDistance
 UCLASS()
 class UEnvQueryTest_Distance : public UEnvQueryTest
 {
-	GENERATED_BODY()
-public:
-	UEnvQueryTest_Distance(const FObjectInitializer& ObjectInitializer = FObjectInitializer::Get());
+	GENERATED_UCLASS_BODY()
 
 	/** testing mode */
 	UPROPERTY(EditDefaultsOnly, Category=Distance)
diff --git a/Engine/Source/Runtime/AIModule/Classes/EnvironmentQuery/Tests/EnvQueryTest_Dot.h b/Engine/Source/Runtime/AIModule/Classes/EnvironmentQuery/Tests/EnvQueryTest_Dot.h
--- a/Engine/Source/Runtime/AIModule/Classes/EnvironmentQuery/Tests/EnvQueryTest_Dot.h
+++ b/Engine/Source/Runtime/AIModule/Classes/EnvironmentQuery/Tests/EnvQueryTest_Dot.h
@@ -18,9 +18,7 @@ enum class EEnvTestDot : uint8
 UCLASS(MinimalAPI)
 class UEnvQueryTest_Dot : public UEnvQueryTest
 {
-	GENERATED_BODY()
-public:
-	AIMODULE_API UEnvQueryTest_Dot(const FObjectInitializer& ObjectInitializer = FObjectInitializer::Get());
+	GENERATED_UCLASS_BODY()
 
 protected:
 	/** defines direction of first line used by test */
diff --git a/Engine/Source/Runtime/AIModule/Classes/EnvironmentQuery/Tests/EnvQueryTest_GameplayTags.h b/Engine/Source/Runtime/AIModule/Classes/EnvironmentQuery/Tests/EnvQueryTest_GameplayTags.h
--- a/Engine/Source/Runtime/AIModule/Classes/EnvironmentQuery/Tests/EnvQueryTest_GameplayTags.h
+++ b/Engine/Source/Runtime/AIModule/Classes/EnvironmentQuery/Tests/EnvQueryTest_GameplayTags.h
@@ -10,9 +10,7 @@ class IGameplayTagAssetInterface;
 UCLASS(MinimalAPI)
 class UEnvQueryTest_GameplayTags : public UEnvQueryTest
 {
-	GENERATED_BODY()
-public:
-	AIMODULE_API UEnvQueryTest_GameplayTags(const FObjectInitializer& ObjectInitializer = FObjectInitializer::Get());
+	GENERATED_UCLASS_BODY()
 
 protected:
 	virtual void RunTest(FEnvQueryInstance& QueryInstance) const override;
diff --git a/Engine/Source/Runtime/AIModule/Classes/EnvironmentQuery/Tests/EnvQueryTest_Pathfinding.h b/Engine/Source/Runtime/AIModule/Classes/EnvironmentQuery/Tests/EnvQueryTest_Pathfinding.h
--- a/Engine/Source/Runtime/AIModule/Classes/EnvironmentQuery/Tests/EnvQueryTest_Pathfinding.h
+++ b/Engine/Source/Runtime/AIModule/Classes/EnvironmentQuery/Tests/EnvQueryTest_Pathfinding.h
@@ -20,9 +20,7 @@ namespace EEnvTestPathfinding
 UCLASS()
 class UEnvQueryTest_Pathfinding : public UEnvQueryTest
 {
-	GENERATED_BODY()
-public:
-	UEnvQueryTest_Pathfinding(const FObjectInitializer& ObjectInitializer = FObjectInitializer::Get());
+	GENERATED_UCLASS_BODY()
 
 	/** testing mode */
 	UPROPERTY(EditDefaultsOnly, Category=Pathfinding)
diff --git a/Engine/Source/Runtime/AIModule/Classes/EnvironmentQuery/Tests/EnvQueryTest_Random.h b/Engine/Source/Runtime/AIModule/Classes/EnvironmentQuery/Tests/EnvQueryTest_Random.h
--- a/Engine/Source/Runtime/AIModule/Classes/EnvironmentQuery/Tests/EnvQueryTest_Random.h
+++ b/Engine/Source/Runtime/AIModule/Classes/EnvironmentQuery/Tests/EnvQueryTest_Random.h
@@ -7,9 +7,7 @@
 UCLASS(MinimalAPI, meta=(DeprecatedNode, DeprecationMessage = "This class is now deprecated, please use RunMode supporting random results instead."))
 class UEnvQueryTest_Random : public UEnvQueryTest
 {
-	GENERATED_BODY()
-public:
-	AIMODULE_API UEnvQueryTest_Random(const FObjectInitializer& ObjectInitializer = FObjectInitializer::Get());
+	GENERATED_UCLASS_BODY()
 
 protected:
 	virtual void RunTest(FEnvQueryInstance& QueryInstance) const override;
diff --git a/Engine/Source/Runtime/AIModule/Classes/EnvironmentQuery/Tests/EnvQueryTest_Trace.h b/Engine/Source/Runtime/AIModule/Classes/EnvironmentQuery/Tests/EnvQueryTest_Trace.h
--- a/Engine/Source/Runtime/AIModule/Classes/EnvironmentQuery/Tests/EnvQueryTest_Trace.h
+++ b/Engine/Source/Runtime/AIModule/Classes/EnvironmentQuery/Tests/EnvQueryTest_Trace.h
@@ -9,9 +9,7 @@ class UEnvQueryContext;
 UCLASS(MinimalAPI)
 class UEnvQueryTest_Trace : public UEnvQueryTest
 {
-	GENERATED_BODY()
-public:
-	AIMODULE_API UEnvQueryTest_Trace(const FObjectInitializer& ObjectInitializer = FObjectInitializer::Get());
+	GENERATED_UCLASS_BODY()
 
 	/** trace data */
 	UPROPERTY(EditDefaultsOnly, Category=Trace)
diff --git a/Engine/Source/Runtime/AIModule/Classes/GenericTeamAgentInterface.h b/Engine/Source/Runtime/AIModule/Classes/GenericTeamAgentInterface.h
--- a/Engine/Source/Runtime/AIModule/Classes/GenericTeamAgentInterface.h
+++ b/Engine/Source/Runtime/AIModule/Classes/GenericTeamAgentInterface.h
@@ -18,7 +18,7 @@ namespace ETeamAttitude
 USTRUCT()
 struct AIMODULE_API FGenericTeamId
 {
-	GENERATED_BODY()
+	GENERATED_USTRUCT_BODY()
 
 private:
 	enum EPredefinedId
@@ -49,15 +49,12 @@ public:
 UINTERFACE()
 class AIMODULE_API UGenericTeamAgentInterface : public UInterface
 {
-	GENERATED_BODY()
-public:
-	UGenericTeamAgentInterface(const FObjectInitializer& ObjectInitializer = FObjectInitializer::Get());
+	GENERATED_UINTERFACE_BODY()
 };
 
 class AIMODULE_API IGenericTeamAgentInterface
 {
-	GENERATED_BODY()
-public:
+	GENERATED_IINTERFACE_BODY()
 
 	/** Assigns Team Agent to given TeamID */
 	virtual void SetGenericTeamId(const FGenericTeamId& TemaID) {}
diff --git a/Engine/Source/Runtime/AIModule/Classes/HotSpots/AIHotSpotManager.h b/Engine/Source/Runtime/AIModule/Classes/HotSpots/AIHotSpotManager.h
--- a/Engine/Source/Runtime/AIModule/Classes/HotSpots/AIHotSpotManager.h
+++ b/Engine/Source/Runtime/AIModule/Classes/HotSpots/AIHotSpotManager.h
@@ -6,7 +6,5 @@
 UCLASS()
 class AIMODULE_API UAIHotSpotManager : public UObject
 {
-	GENERATED_BODY()
-public:
-	UAIHotSpotManager(const FObjectInitializer& ObjectInitializer = FObjectInitializer::Get());
-};
+	GENERATED_UCLASS_BODY()
+};
\ No newline at end of file
diff --git a/Engine/Source/Runtime/AIModule/Classes/Navigation/CrowdAgentInterface.h b/Engine/Source/Runtime/AIModule/Classes/Navigation/CrowdAgentInterface.h
--- a/Engine/Source/Runtime/AIModule/Classes/Navigation/CrowdAgentInterface.h
+++ b/Engine/Source/Runtime/AIModule/Classes/Navigation/CrowdAgentInterface.h
@@ -7,15 +7,12 @@
 UINTERFACE(MinimalAPI, meta=(CannotImplementInterfaceInBlueprint))
 class UCrowdAgentInterface : public UInterface
 {
-	GENERATED_BODY()
-public:
-	AIMODULE_API UCrowdAgentInterface(const FObjectInitializer& ObjectInitializer = FObjectInitializer::Get());
+	GENERATED_UINTERFACE_BODY()
 };
 
 class ICrowdAgentInterface
 {
-	GENERATED_BODY()
-public:
+	GENERATED_IINTERFACE_BODY()
 
 	/** @return current location of crowd agent */
 	virtual FVector GetCrowdAgentLocation() const { return FAISystem::InvalidLocation; }
diff --git a/Engine/Source/Runtime/AIModule/Classes/Perception/AIPerceptionComponent.h b/Engine/Source/Runtime/AIModule/Classes/Perception/AIPerceptionComponent.h
--- a/Engine/Source/Runtime/AIModule/Classes/Perception/AIPerceptionComponent.h
+++ b/Engine/Source/Runtime/AIModule/Classes/Perception/AIPerceptionComponent.h
@@ -84,7 +84,7 @@ struct AIMODULE_API FActorPerceptionInfo
 USTRUCT(BlueprintType, meta = (DisplayName = "Sensed Actor's Data"))
 struct FActorPerceptionBlueprintInfo
 {
-	GENERATED_BODY()
+	GENERATED_USTRUCT_BODY()
 
 	UPROPERTY(BlueprintReadWrite, Category = "AI|Perception")
 	AActor* Target;
@@ -107,9 +107,7 @@ struct FActorPerceptionBlueprintInfo
 UCLASS(ClassGroup=AI, HideCategories=(Activation, Collision), meta=(BlueprintSpawnableComponent), config=Game)
 class AIMODULE_API UAIPerceptionComponent : public UActorComponent
 {
-	GENERATED_BODY()
-public:
-	UAIPerceptionComponent(const FObjectInitializer& ObjectInitializer = FObjectInitializer::Get());
+	GENERATED_UCLASS_BODY()
 	
 	static const int32 InitialStimuliToProcessArraySize;
 
diff --git a/Engine/Source/Runtime/AIModule/Classes/Perception/AIPerceptionListenerInterface.h b/Engine/Source/Runtime/AIModule/Classes/Perception/AIPerceptionListenerInterface.h
--- a/Engine/Source/Runtime/AIModule/Classes/Perception/AIPerceptionListenerInterface.h
+++ b/Engine/Source/Runtime/AIModule/Classes/Perception/AIPerceptionListenerInterface.h
@@ -9,15 +9,12 @@ class UAIPerceptionComponent;
 UINTERFACE()
 class AIMODULE_API UAIPerceptionListenerInterface : public UInterface
 {
-	GENERATED_BODY()
-public:
-	UAIPerceptionListenerInterface(const FObjectInitializer& ObjectInitializer = FObjectInitializer::Get());
+	GENERATED_UINTERFACE_BODY()
 };
 
 class AIMODULE_API IAIPerceptionListenerInterface
 {
-	GENERATED_BODY()
-public:
+	GENERATED_IINTERFACE_BODY()
 
 	virtual UAIPerceptionComponent* GetPerceptionComponent() { return NULL; }
 };
diff --git a/Engine/Source/Runtime/AIModule/Classes/Perception/AIPerceptionTypes.h b/Engine/Source/Runtime/AIModule/Classes/Perception/AIPerceptionTypes.h
--- a/Engine/Source/Runtime/AIModule/Classes/Perception/AIPerceptionTypes.h
+++ b/Engine/Source/Runtime/AIModule/Classes/Perception/AIPerceptionTypes.h
@@ -61,7 +61,7 @@ struct FPerceptionChannelFilter
 USTRUCT(BlueprintType)
 struct AIMODULE_API FAIStimulus
 {
-	GENERATED_BODY()
+	GENERATED_USTRUCT_BODY()
 
 	static const float NeverHappenedAge;
 
@@ -130,7 +130,7 @@ public:
 USTRUCT()
 struct AIMODULE_API FAISenseAffiliationFilter
 {
-	GENERATED_BODY()
+	GENERATED_USTRUCT_BODY()
 
 	UPROPERTY(EditDefaultsOnly, BlueprintReadOnly, Category = "Sense")
 	uint32 bDetectEnemies : 1;
diff --git a/Engine/Source/Runtime/AIModule/Classes/Perception/AISense.h b/Engine/Source/Runtime/AIModule/Classes/Perception/AISense.h
--- a/Engine/Source/Runtime/AIModule/Classes/Perception/AISense.h
+++ b/Engine/Source/Runtime/AIModule/Classes/Perception/AISense.h
@@ -13,9 +13,7 @@ DECLARE_DELEGATE_OneParam(FOnPerceptionListenerUpdateDelegate, const FPerception
 UCLASS(ClassGroup = AI, abstract, config = Engine)
 class AIMODULE_API UAISense : public UObject
 {
-	GENERATED_BODY()
-public:
-	UAISense(const FObjectInitializer& ObjectInitializer = FObjectInitializer::Get());
+	GENERATED_UCLASS_BODY()
 
 	static const float SuspendNextUpdate;
 
diff --git a/Engine/Source/Runtime/AIModule/Classes/Perception/AISense_Damage.h b/Engine/Source/Runtime/AIModule/Classes/Perception/AISense_Damage.h
--- a/Engine/Source/Runtime/AIModule/Classes/Perception/AISense_Damage.h
+++ b/Engine/Source/Runtime/AIModule/Classes/Perception/AISense_Damage.h
@@ -9,7 +9,7 @@
 USTRUCT()
 struct AIMODULE_API FAIDamageEvent
 {	
-	GENERATED_BODY()
+	GENERATED_USTRUCT_BODY()
 
 	typedef class UAISense_Damage FSenseClass;
 
@@ -49,9 +49,7 @@ struct AIMODULE_API FAIDamageEvent
 UCLASS(ClassGroup=AI)
 class AIMODULE_API UAISense_Damage : public UAISense
 {
-	GENERATED_BODY()
-public:
-	UAISense_Damage(const FObjectInitializer& ObjectInitializer = FObjectInitializer::Get());
+	GENERATED_UCLASS_BODY()
 
 	UPROPERTY()
 	TArray<FAIDamageEvent> RegisteredEvents;
diff --git a/Engine/Source/Runtime/AIModule/Classes/Perception/AISense_Prediction.h b/Engine/Source/Runtime/AIModule/Classes/Perception/AISense_Prediction.h
--- a/Engine/Source/Runtime/AIModule/Classes/Perception/AISense_Prediction.h
+++ b/Engine/Source/Runtime/AIModule/Classes/Perception/AISense_Prediction.h
@@ -11,7 +11,7 @@ class UAISense_Prediction;
 USTRUCT()
 struct AIMODULE_API FAIPredictionEvent
 {	
-	GENERATED_BODY()
+	GENERATED_USTRUCT_BODY()
 
 	typedef UAISense_Prediction FSenseClass;
 	
@@ -34,9 +34,7 @@ struct AIMODULE_API FAIPredictionEvent
 UCLASS(ClassGroup=AI)
 class AIMODULE_API UAISense_Prediction : public UAISense
 {
-	GENERATED_BODY()
-public:
-	UAISense_Prediction(const FObjectInitializer& ObjectInitializer = FObjectInitializer::Get());
+	GENERATED_UCLASS_BODY()
 
 	UPROPERTY()
 	TArray<FAIPredictionEvent> RegisteredEvents;
diff --git a/Engine/Source/Runtime/AIModule/Classes/Perception/AISense_Sight.h b/Engine/Source/Runtime/AIModule/Classes/Perception/AISense_Sight.h
--- a/Engine/Source/Runtime/AIModule/Classes/Perception/AISense_Sight.h
+++ b/Engine/Source/Runtime/AIModule/Classes/Perception/AISense_Sight.h
@@ -23,7 +23,7 @@ namespace ESightPerceptionEventName
 USTRUCT()
 struct AIMODULE_API FAISightEvent
 {
-	GENERATED_BODY()
+	GENERATED_USTRUCT_BODY()
 
 	typedef UAISense_Sight FSenseClass;
 
@@ -101,9 +101,7 @@ struct FAISightQuery
 UCLASS(ClassGroup=AI, config=Game)
 class AIMODULE_API UAISense_Sight : public UAISense
 {
-	GENERATED_BODY()
-public:
-	UAISense_Sight(const FObjectInitializer& ObjectInitializer = FObjectInitializer::Get());
+	GENERATED_UCLASS_BODY()
 
 public:
 	struct FDigestedSightProperties
diff --git a/Engine/Source/Runtime/AIModule/Classes/Perception/AISense_Team.h b/Engine/Source/Runtime/AIModule/Classes/Perception/AISense_Team.h
--- a/Engine/Source/Runtime/AIModule/Classes/Perception/AISense_Team.h
+++ b/Engine/Source/Runtime/AIModule/Classes/Perception/AISense_Team.h
@@ -11,7 +11,7 @@ class UAISense_Team;
 USTRUCT()
 struct AIMODULE_API FAITeamStimulusEvent
 {	
-	GENERATED_BODY()
+	GENERATED_USTRUCT_BODY()
 
 	typedef UAISense_Team FSenseClass;
 
@@ -47,9 +47,7 @@ public:
 UCLASS(ClassGroup=AI)
 class AIMODULE_API UAISense_Team : public UAISense
 {
-	GENERATED_BODY()
-public:
-	UAISense_Team(const FObjectInitializer& ObjectInitializer = FObjectInitializer::Get());
+	GENERATED_UCLASS_BODY()
 
 	UPROPERTY()
 	TArray<FAITeamStimulusEvent> RegisteredEvents;
diff --git a/Engine/Source/Runtime/AIModule/Classes/Perception/AISense_Touch.h b/Engine/Source/Runtime/AIModule/Classes/Perception/AISense_Touch.h
--- a/Engine/Source/Runtime/AIModule/Classes/Perception/AISense_Touch.h
+++ b/Engine/Source/Runtime/AIModule/Classes/Perception/AISense_Touch.h
@@ -11,7 +11,7 @@ class UAISense_Touch;
 USTRUCT()
 struct AIMODULE_API FAITouchEvent
 {	
-	GENERATED_BODY()
+	GENERATED_USTRUCT_BODY()
 
 	typedef UAISense_Touch FSenseClass;
 
@@ -34,9 +34,7 @@ struct AIMODULE_API FAITouchEvent
 UCLASS(ClassGroup=AI)
 class AIMODULE_API UAISense_Touch : public UAISense
 {
-	GENERATED_BODY()
-public:
-	UAISense_Touch(const FObjectInitializer& ObjectInitializer = FObjectInitializer::Get());
+	GENERATED_UCLASS_BODY()
 
 	UPROPERTY()
 	TArray<FAITouchEvent> RegisteredEvents;
diff --git a/Engine/Source/Runtime/AIModule/Classes/Perception/AISightTargetInterface.h b/Engine/Source/Runtime/AIModule/Classes/Perception/AISightTargetInterface.h
--- a/Engine/Source/Runtime/AIModule/Classes/Perception/AISightTargetInterface.h
+++ b/Engine/Source/Runtime/AIModule/Classes/Perception/AISightTargetInterface.h
@@ -7,15 +7,12 @@
 UINTERFACE()
 class AIMODULE_API UAISightTargetInterface : public UInterface
 {
-	GENERATED_BODY()
-public:
-	UAISightTargetInterface(const FObjectInitializer& ObjectInitializer = FObjectInitializer::Get());
+	GENERATED_UINTERFACE_BODY()
 };
 
 class AIMODULE_API IAISightTargetInterface
 {
-	GENERATED_BODY()
-public:
+	GENERATED_IINTERFACE_BODY()
 
 	/**	Implementation should check whether from given ObserverLocation
 	 *	implementer can be seen. If so OutSeenLocation should contain
diff --git a/Engine/Source/Runtime/AIModule/Classes/Perception/PawnSensingComponent.h b/Engine/Source/Runtime/AIModule/Classes/Perception/PawnSensingComponent.h
--- a/Engine/Source/Runtime/AIModule/Classes/Perception/PawnSensingComponent.h
+++ b/Engine/Source/Runtime/AIModule/Classes/Perception/PawnSensingComponent.h
@@ -16,9 +16,7 @@ class APawn;
 UCLASS(ClassGroup=AI, HideCategories=(Activation, "Components|Activation", Collision), meta=(BlueprintSpawnableComponent))
 class AIMODULE_API UPawnSensingComponent : public UActorComponent
 {
-	GENERATED_BODY()
-public:
-	UPawnSensingComponent(const FObjectInitializer& ObjectInitializer = FObjectInitializer::Get());
+	GENERATED_UCLASS_BODY()
 
 	DECLARE_DYNAMIC_MULTICAST_DELEGATE_OneParam( FSeePawnDelegate, APawn*, Pawn );
 	DECLARE_DYNAMIC_MULTICAST_DELEGATE_ThreeParams( FHearNoiseDelegate, APawn*, Instigator, const FVector&, Location, float, Volume);
diff --git a/Engine/Source/Runtime/Android/AndroidRuntimeSettings/Classes/AndroidRuntimeSettings.h b/Engine/Source/Runtime/Android/AndroidRuntimeSettings/Classes/AndroidRuntimeSettings.h
--- a/Engine/Source/Runtime/Android/AndroidRuntimeSettings/Classes/AndroidRuntimeSettings.h
+++ b/Engine/Source/Runtime/Android/AndroidRuntimeSettings/Classes/AndroidRuntimeSettings.h
@@ -56,7 +56,7 @@ namespace EAndroidDepthBufferPreference
 USTRUCT()
 struct FGooglePlayAchievementMapping
 {
-	GENERATED_BODY()
+	GENERATED_USTRUCT_BODY()
 
 	// The game-specific achievement name (the one passed in to WriteAchievement calls)
 	UPROPERTY(EditAnywhere, Category = GooglePlayServices)
@@ -73,7 +73,7 @@ struct FGooglePlayAchievementMapping
 USTRUCT()
 struct FGooglePlayLeaderboardMapping
 {
-	GENERATED_BODY()
+	GENERATED_USTRUCT_BODY()
 
 	// The game-specific leaderboard name (the one passed in to WriteLeaderboards calls)
 	UPROPERTY(EditAnywhere, Category = GooglePlayServices)
@@ -101,10 +101,8 @@ namespace EAndroidAudio
 UCLASS(config=Engine, defaultconfig)
 class ANDROIDRUNTIMESETTINGS_API UAndroidRuntimeSettings : public UObject
 {
-	GENERATED_BODY()
-public:
-	UAndroidRuntimeSettings(const FObjectInitializer& ObjectInitializer = FObjectInitializer::Get());
 public:
+	GENERATED_UCLASS_BODY()
 
 	// The official name of the product (same as the name you use on the Play Store web site)
 	UPROPERTY(GlobalConfig, EditAnywhere, Category = APKPackaging, Meta = (DisplayName = "Android Package Name ('com.Company.Project', [PROJECT] is replaced with project name)"))
diff --git a/Engine/Source/Runtime/AutomationMessages/Classes/AutomationWorkerMessages.h b/Engine/Source/Runtime/AutomationMessages/Classes/AutomationWorkerMessages.h
--- a/Engine/Source/Runtime/AutomationMessages/Classes/AutomationWorkerMessages.h
+++ b/Engine/Source/Runtime/AutomationMessages/Classes/AutomationWorkerMessages.h
@@ -14,7 +14,7 @@
 USTRUCT()
 struct FAutomationWorkerFindWorkers
 {
-	GENERATED_BODY()
+	GENERATED_USTRUCT_BODY()
 
 	/** Holds the change list number to find workers for. */
 	UPROPERTY()
@@ -61,7 +61,7 @@ struct TStructOpsTypeTraits<FAutomationWorkerFindWorkers> : public TStructOpsTyp
 USTRUCT()
 struct FAutomationWorkerFindWorkersResponse
 {
-	GENERATED_BODY()
+	GENERATED_USTRUCT_BODY()
 
 	/** Holds the name of the device that the worker is running on. */
 	UPROPERTY()
@@ -125,7 +125,7 @@ struct TStructOpsTypeTraits<FAutomationWorkerFindWorkersResponse> : public TStru
 USTRUCT()
 struct FAutomationWorkerWorkerOffline
 {
-	GENERATED_BODY()
+	GENERATED_USTRUCT_BODY()
 };
 
 template<>
@@ -143,7 +143,7 @@ struct TStructOpsTypeTraits<FAutomationWorkerWorkerOffline> : public TStructOpsT
 USTRUCT()
 struct FAutomationWorkerPing
 {
-	GENERATED_BODY()
+	GENERATED_USTRUCT_BODY()
 };
 
 template<>
@@ -161,7 +161,7 @@ struct TStructOpsTypeTraits<FAutomationWorkerPing> : public TStructOpsTypeTraits
 USTRUCT()
 struct FAutomationWorkerResetTests
 {
-	GENERATED_BODY()
+	GENERATED_USTRUCT_BODY()
 };
 
 template<>
@@ -179,7 +179,7 @@ struct TStructOpsTypeTraits<FAutomationWorkerResetTests> : public TStructOpsType
 USTRUCT()
 struct FAutomationWorkerPong
 {
-	GENERATED_BODY()
+	GENERATED_USTRUCT_BODY()
 };
 
 template<>
@@ -198,7 +198,7 @@ struct TStructOpsTypeTraits<FAutomationWorkerPong> : public TStructOpsTypeTraits
 USTRUCT()
 struct FAutomationWorkerRequestTests
 {
-	GENERATED_BODY()
+	GENERATED_USTRUCT_BODY()
 
 	/** Holds a flag indicating whether the developer directory should be included. */
 	UPROPERTY()
@@ -238,7 +238,7 @@ struct TStructOpsTypeTraits<FAutomationWorkerRequestTests> : public TStructOpsTy
 USTRUCT()
 struct FAutomationWorkerRequestTestsReply
 {
-	GENERATED_BODY()
+	GENERATED_USTRUCT_BODY()
 
 	/** Holds the test information serialized into a string. */
 	UPROPERTY()
@@ -278,7 +278,7 @@ struct TStructOpsTypeTraits<FAutomationWorkerRequestTestsReply> : public TStruct
 USTRUCT()
 struct FAutomationWorkerRunTests
 {
-	GENERATED_BODY()
+	GENERATED_USTRUCT_BODY()
 
 	/** */
 	UPROPERTY()
@@ -333,7 +333,7 @@ struct TStructOpsTypeTraits<FAutomationWorkerRunTests> : public TStructOpsTypeTr
 USTRUCT()
 struct FAutomationWorkerRunTestsReply
 {
-	GENERATED_BODY()
+	GENERATED_USTRUCT_BODY()
 
 	/** */
 	UPROPERTY()
@@ -379,7 +379,7 @@ struct TStructOpsTypeTraits<FAutomationWorkerRunTestsReply> : public TStructOpsT
 USTRUCT()
 struct FAutomationWorkerRequestNextNetworkCommand
 {
-	GENERATED_BODY()
+	GENERATED_USTRUCT_BODY()
 
 	/** */
 	UPROPERTY()
@@ -413,7 +413,7 @@ struct TStructOpsTypeTraits<FAutomationWorkerRequestNextNetworkCommand> : public
 USTRUCT()
 struct FAutomationWorkerNextNetworkCommandReply
 {
-	GENERATED_BODY()
+	GENERATED_USTRUCT_BODY()
 };
 
 template<>
@@ -432,7 +432,7 @@ struct TStructOpsTypeTraits<FAutomationWorkerNextNetworkCommandReply> : public T
 USTRUCT()
 struct FAutomationWorkerScreenImage
 {
-	GENERATED_BODY()
+	GENERATED_USTRUCT_BODY()
 
 	/** The screen shot data. */
 	UPROPERTY()
diff --git a/Engine/Source/Runtime/Core/Public/GenericPlatform/GenericPlatformCompilerPreSetup.h b/Engine/Source/Runtime/Core/Public/GenericPlatform/GenericPlatformCompilerPreSetup.h
--- a/Engine/Source/Runtime/Core/Public/GenericPlatform/GenericPlatformCompilerPreSetup.h
+++ b/Engine/Source/Runtime/Core/Public/GenericPlatform/GenericPlatformCompilerPreSetup.h
@@ -79,9 +79,6 @@
 	#define EMIT_DEPRECATED_WARNING_MESSAGE(Msg)
 #endif // EMIT_DEPRECATED_WARNING_MESSAGE
 
-#ifndef EMIT_CUSTOM_WARNING_AT_LINE
-	#define EMIT_CUSTOM_WARNING_AT_LINE(Line, Warning)
-#endif
-
-#define EMIT_CUSTOM_WARNING(Warning) \
-	EMIT_CUSTOM_WARNING_AT_LINE(__LINE__, Warning)
+#ifndef EMIT_CUSTOM_WARNING
+	#define EMIT_CUSTOM_WARNING(Warning)
+#endif // EMIT_CUSTOM_WARNING
\ No newline at end of file
diff --git a/Engine/Source/Runtime/Core/Public/HTML5/HTML5PlatformCompilerPreSetup.h b/Engine/Source/Runtime/Core/Public/HTML5/HTML5PlatformCompilerPreSetup.h
--- a/Engine/Source/Runtime/Core/Public/HTML5/HTML5PlatformCompilerPreSetup.h
+++ b/Engine/Source/Runtime/Core/Public/HTML5/HTML5PlatformCompilerPreSetup.h
@@ -83,7 +83,7 @@
 #endif 
 #endif // DISABLE_DEPRECATION
 
-#ifndef EMIT_CUSTOM_WARNING_AT_LINE
-#define EMIT_CUSTOM_WARNING_AT_LINE(Line, Warning) \
-	_Pragma(PREPROCESSOR_TO_STRING(message(WARNING_LOCATION(Line) Warning)))
-#endif // EMIT_CUSTOM_WARNING_AT_LINE
+#ifndef EMIT_CUSTOM_WARNING
+#define EMIT_CUSTOM_WARNING(Warning) \
+	_Pragma(PREPROCESSOR_TO_STRING(message(WARNING_LOCATION Warning)))
+#endif // EMIT_CUSTOM_WARNING
diff --git a/Engine/Source/Runtime/Core/Public/WinRT/WinRTPlatformCompilerPreSetup.h b/Engine/Source/Runtime/Core/Public/WinRT/WinRTPlatformCompilerPreSetup.h
--- a/Engine/Source/Runtime/Core/Public/WinRT/WinRTPlatformCompilerPreSetup.h
+++ b/Engine/Source/Runtime/Core/Public/WinRT/WinRTPlatformCompilerPreSetup.h
@@ -89,11 +89,11 @@
 #endif // DISABLE_DEPRECATION
 
 #if _MSC_VER
-#define EMIT_CUSTOM_WARNING_AT_LINE(Line, Warning) \
-	__pragma(message(WARNING_LOCATION(Line) ": warning C4996: " Warning))
+#define EMIT_CUSTOM_WARNING(Warning) \
+	__pragma(message(WARNING_LOCATION "(\": warning C4996: " Warning))
 #elif __clang__
-#define EMIT_CUSTOM_WARNING_AT_LINE(Line, Warning) \
-	_Pragma(PREPROCESSOR_TO_STRING(message(WARNING_LOCATION(Line) Warning)))
+#define EMIT_CUSTOM_WARNING(Warning) \
+	_Pragma(PREPROCESSOR_TO_STRING(message(WARNING_LOCATION Warning)))
 #else
 #error Unknown compiler
 #endif
diff --git a/Engine/Source/Runtime/CoreUObject/Public/UObject/ObjectBase.h b/Engine/Source/Runtime/CoreUObject/Public/UObject/ObjectBase.h
--- a/Engine/Source/Runtime/CoreUObject/Public/UObject/ObjectBase.h
+++ b/Engine/Source/Runtime/CoreUObject/Public/UObject/ObjectBase.h
@@ -515,7 +515,7 @@ struct COREUOBJECT_API FReferencerInformationList
 #define UENUM(...)
 #define UDELEGATE(...)
 
-// This pair of macros is used to help implement GENERATED*_BODY() macros
+// This pair of macros is used to help implement GENERATED_BODY() and GENERATED_USTRUCT_BODY()
 #define BODY_MACRO_COMBINE_INNER(A,B,C,D) A##B##C##D
 #define BODY_MACRO_COMBINE(A,B,C,D) BODY_MACRO_COMBINE_INNER(A,B,C,D)
 
@@ -525,7 +525,7 @@ struct COREUOBJECT_API FReferencerInformationList
 #define GENERATED_USTRUCT_BODY(...) GENERATED_BODY()
 #define GENERATED_UCLASS_BODY(...) GENERATED_BODY_LEGACY()
 #define GENERATED_UINTERFACE_BODY(...) GENERATED_BODY_LEGACY()
-#define GENERATED_IINTERFACE_BODY(...) GENERATED_BODY_LEGACY()
+#define GENERATED_IINTERFACE_BODY(...) GENERATED_BODY()
 
 #if UE_BUILD_DOCS
 #define UCLASS(...)
diff --git a/Engine/Source/Runtime/Engine/Classes/AI/AISystemBase.h b/Engine/Source/Runtime/Engine/Classes/AI/AISystemBase.h
--- a/Engine/Source/Runtime/Engine/Classes/AI/AISystemBase.h
+++ b/Engine/Source/Runtime/Engine/Classes/AI/AISystemBase.h
@@ -9,9 +9,7 @@ class UWorld;
 UCLASS(abstract, config = Engine, defaultconfig)
 class ENGINE_API UAISystemBase : public UObject
 {
-	GENERATED_BODY()
-public:
-	UAISystemBase(const FObjectInitializer& ObjectInitializer = FObjectInitializer::Get());
+	GENERATED_UCLASS_BODY()
 
 	virtual ~UAISystemBase(){}
 		
@@ -52,4 +50,4 @@ class IAISystemModule : public IModuleInterface
 {
 public:
 	virtual UAISystemBase* CreateAISystemInstance(UWorld* World) = 0;
-};
+};
\ No newline at end of file
diff --git a/Engine/Source/Runtime/Engine/Classes/AI/Navigation/AbstractNavData.h b/Engine/Source/Runtime/Engine/Classes/AI/Navigation/AbstractNavData.h
--- a/Engine/Source/Runtime/Engine/Classes/AI/Navigation/AbstractNavData.h
+++ b/Engine/Source/Runtime/Engine/Classes/AI/Navigation/AbstractNavData.h
@@ -36,9 +36,7 @@ public:
 UCLASS()
 class AAbstractNavData : public ANavigationData
 {
-	GENERATED_BODY()
-public:
-	AAbstractNavData(const FObjectInitializer& ObjectInitializer = FObjectInitializer::Get());
+	GENERATED_UCLASS_BODY()
 
 	// Begin ANavigationData overrides
 	virtual void BatchRaycast(TArray<FNavigationRaycastWork>& Workload, TSharedPtr<const FNavigationQueryFilter> QueryFilter, const UObject* Querier = NULL) const override {};
diff --git a/Engine/Source/Runtime/Engine/Classes/AI/Navigation/AvoidanceManager.h b/Engine/Source/Runtime/Engine/Classes/AI/Navigation/AvoidanceManager.h
--- a/Engine/Source/Runtime/Engine/Classes/AI/Navigation/AvoidanceManager.h
+++ b/Engine/Source/Runtime/Engine/Classes/AI/Navigation/AvoidanceManager.h
@@ -12,7 +12,7 @@ class UAvoidanceManager;
 USTRUCT()
 struct FNavAvoidanceData
 {
-	GENERATED_BODY()
+	GENERATED_USTRUCT_BODY()
 
 	/** Current location */
 	FVector Center;
@@ -75,9 +75,7 @@ struct FVelocityAvoidanceCone
 UCLASS(config=Engine, Blueprintable)
 class ENGINE_API UAvoidanceManager : public UObject, public FSelfRegisteringExec
 {
-	GENERATED_BODY()
-public:
-	UAvoidanceManager(const FObjectInitializer& ObjectInitializer = FObjectInitializer::Get());
+	GENERATED_UCLASS_BODY()
 		
 	/** How long an avoidance UID must not be updated before the system will put it back in the pool. Actual delay is up to 150% of this value. */
 	UPROPERTY(EditAnywhere, Category="Avoidance", config, meta=(ClampMin = "0.0"))
diff --git a/Engine/Source/Runtime/Engine/Classes/AI/Navigation/NavAgentInterface.h b/Engine/Source/Runtime/Engine/Classes/AI/Navigation/NavAgentInterface.h
--- a/Engine/Source/Runtime/Engine/Classes/AI/Navigation/NavAgentInterface.h
+++ b/Engine/Source/Runtime/Engine/Classes/AI/Navigation/NavAgentInterface.h
@@ -11,15 +11,12 @@ struct FNavAgentProperties;
 UINTERFACE(MinimalAPI, meta=(CannotImplementInterfaceInBlueprint))
 class UNavAgentInterface : public UInterface
 {
-	GENERATED_BODY()
-public:
-	ENGINE_API UNavAgentInterface(const FObjectInitializer& ObjectInitializer = FObjectInitializer::Get());
+	GENERATED_UINTERFACE_BODY()
 };
 
 class INavAgentInterface
 {
-	GENERATED_BODY()
-public:
+	GENERATED_IINTERFACE_BODY()
 
 	/**
 	 *	Retrieves FNavAgentProperties expressing navigation props and caps of represented agent
diff --git a/Engine/Source/Runtime/Engine/Classes/AI/Navigation/NavAreas/NavArea.h b/Engine/Source/Runtime/Engine/Classes/AI/Navigation/NavAreas/NavArea.h
--- a/Engine/Source/Runtime/Engine/Classes/AI/Navigation/NavAreas/NavArea.h
+++ b/Engine/Source/Runtime/Engine/Classes/AI/Navigation/NavAreas/NavArea.h
@@ -7,9 +7,7 @@
 UCLASS(DefaultToInstanced, abstract, Config=Engine, Blueprintable)
 class ENGINE_API UNavArea : public UObject
 {
-	GENERATED_BODY()
-public:
-	UNavArea(const FObjectInitializer& ObjectInitializer = FObjectInitializer::Get());
+	GENERATED_UCLASS_BODY()
 
 	/** travel cost multiplier for path distance */
 	UPROPERTY(EditAnywhere, Category=NavArea, config, meta=(ClampMin = "0.0"))
diff --git a/Engine/Source/Runtime/Engine/Classes/AI/Navigation/NavAreas/NavAreaMeta_SwitchByAgent.h b/Engine/Source/Runtime/Engine/Classes/AI/Navigation/NavAreas/NavAreaMeta_SwitchByAgent.h
--- a/Engine/Source/Runtime/Engine/Classes/AI/Navigation/NavAreas/NavAreaMeta_SwitchByAgent.h
+++ b/Engine/Source/Runtime/Engine/Classes/AI/Navigation/NavAreas/NavAreaMeta_SwitchByAgent.h
@@ -9,9 +9,7 @@
 UCLASS(Abstract)
 class ENGINE_API UNavAreaMeta_SwitchByAgent : public UNavAreaMeta
 {
-	GENERATED_BODY()
-public:
-	UNavAreaMeta_SwitchByAgent(const FObjectInitializer& ObjectInitializer = FObjectInitializer::Get());
+	GENERATED_UCLASS_BODY()
 
 	UPROPERTY(EditAnywhere, Category=AgentTypes)
 	TSubclassOf<UNavArea> Agent0Area;
diff --git a/Engine/Source/Runtime/Engine/Classes/AI/Navigation/NavAreas/NavArea_Default.h b/Engine/Source/Runtime/Engine/Classes/AI/Navigation/NavAreas/NavArea_Default.h
--- a/Engine/Source/Runtime/Engine/Classes/AI/Navigation/NavAreas/NavArea_Default.h
+++ b/Engine/Source/Runtime/Engine/Classes/AI/Navigation/NavAreas/NavArea_Default.h
@@ -8,7 +8,5 @@
 UCLASS(Config=Engine)
 class ENGINE_API UNavArea_Default : public UNavArea
 {
-	GENERATED_BODY()
-public:
-	UNavArea_Default(const FObjectInitializer& ObjectInitializer = FObjectInitializer::Get());
+	GENERATED_UCLASS_BODY()
 };
diff --git a/Engine/Source/Runtime/Engine/Classes/AI/Navigation/NavAreas/NavArea_LowHeight.h b/Engine/Source/Runtime/Engine/Classes/AI/Navigation/NavAreas/NavArea_LowHeight.h
--- a/Engine/Source/Runtime/Engine/Classes/AI/Navigation/NavAreas/NavArea_LowHeight.h
+++ b/Engine/Source/Runtime/Engine/Classes/AI/Navigation/NavAreas/NavArea_LowHeight.h
@@ -8,7 +8,5 @@
 UCLASS(Config = Engine)
 class ENGINE_API UNavArea_LowHeight : public UNavArea
 {
-	GENERATED_BODY()
-public:
-	UNavArea_LowHeight(const FObjectInitializer& ObjectInitializer = FObjectInitializer::Get());
+	GENERATED_UCLASS_BODY()
 };
diff --git a/Engine/Source/Runtime/Engine/Classes/AI/Navigation/NavAreas/NavArea_Null.h b/Engine/Source/Runtime/Engine/Classes/AI/Navigation/NavAreas/NavArea_Null.h
--- a/Engine/Source/Runtime/Engine/Classes/AI/Navigation/NavAreas/NavArea_Null.h
+++ b/Engine/Source/Runtime/Engine/Classes/AI/Navigation/NavAreas/NavArea_Null.h
@@ -8,7 +8,5 @@
 UCLASS(Config=Engine)
 class ENGINE_API UNavArea_Null : public UNavArea
 {
-	GENERATED_BODY()
-public:
-	UNavArea_Null(const FObjectInitializer& ObjectInitializer = FObjectInitializer::Get());
+	GENERATED_UCLASS_BODY()
 };
diff --git a/Engine/Source/Runtime/Engine/Classes/AI/Navigation/NavAreas/NavArea_Obstacle.h b/Engine/Source/Runtime/Engine/Classes/AI/Navigation/NavAreas/NavArea_Obstacle.h
--- a/Engine/Source/Runtime/Engine/Classes/AI/Navigation/NavAreas/NavArea_Obstacle.h
+++ b/Engine/Source/Runtime/Engine/Classes/AI/Navigation/NavAreas/NavArea_Obstacle.h
@@ -8,7 +8,5 @@
 UCLASS(Config = Engine)
 class ENGINE_API UNavArea_Obstacle : public UNavArea
 {
-	GENERATED_BODY()
-public:
-	UNavArea_Obstacle(const FObjectInitializer& ObjectInitializer = FObjectInitializer::Get());
+	GENERATED_UCLASS_BODY()
 };
diff --git a/Engine/Source/Runtime/Engine/Classes/AI/Navigation/NavCollision.h b/Engine/Source/Runtime/Engine/Classes/AI/Navigation/NavCollision.h
--- a/Engine/Source/Runtime/Engine/Classes/AI/Navigation/NavCollision.h
+++ b/Engine/Source/Runtime/Engine/Classes/AI/Navigation/NavCollision.h
@@ -11,7 +11,7 @@ struct FCompositeNavModifier;
 USTRUCT()
 struct FNavCollisionCylinder
 {
-	GENERATED_BODY()
+	GENERATED_USTRUCT_BODY()
 
 	UPROPERTY(EditAnywhere, Category=Cylinder)
 	FVector Offset;
@@ -26,7 +26,7 @@ struct FNavCollisionCylinder
 USTRUCT()
 struct FNavCollisionBox
 {
-	GENERATED_BODY()
+	GENERATED_USTRUCT_BODY()
 
 	UPROPERTY(EditAnywhere, Category=Box)
 	FVector Offset;
@@ -44,9 +44,7 @@ struct FNavCollisionConvex
 UCLASS(config=Engine)
 class ENGINE_API UNavCollision : public UObject
 {
-	GENERATED_BODY()
-public:
-	UNavCollision(const FObjectInitializer& ObjectInitializer = FObjectInitializer::Get());
+	GENERATED_UCLASS_BODY()
 
 	FNavCollisionConvex TriMeshCollision;
 	FNavCollisionConvex ConvexCollision;
diff --git a/Engine/Source/Runtime/Engine/Classes/AI/Navigation/NavFilters/NavigationQueryFilter.h b/Engine/Source/Runtime/Engine/Classes/AI/Navigation/NavFilters/NavigationQueryFilter.h
--- a/Engine/Source/Runtime/Engine/Classes/AI/Navigation/NavFilters/NavigationQueryFilter.h
+++ b/Engine/Source/Runtime/Engine/Classes/AI/Navigation/NavFilters/NavigationQueryFilter.h
@@ -10,7 +10,7 @@ struct FNavigationQueryFilter;
 USTRUCT()
 struct ENGINE_API FNavigationFilterArea
 {
-	GENERATED_BODY()
+	GENERATED_USTRUCT_BODY()
 
 	/** navigation area class */
 	UPROPERTY(EditAnywhere, Category=Area)
@@ -43,7 +43,7 @@ struct ENGINE_API FNavigationFilterArea
 USTRUCT()
 struct ENGINE_API FNavigationFilterFlags
 {
-	GENERATED_BODY()
+	GENERATED_USTRUCT_BODY()
 
 #if CPP
 	union
@@ -201,9 +201,7 @@ protected:
 UCLASS(Abstract, Blueprintable)
 class ENGINE_API UNavigationQueryFilter : public UObject
 {
-	GENERATED_BODY()
-public:
-	UNavigationQueryFilter(const FObjectInitializer& ObjectInitializer = FObjectInitializer::Get());
+	GENERATED_UCLASS_BODY()
 	
 	/** list of overrides for navigation areas */
 	UPROPERTY(EditAnywhere, Category=Filter)
diff --git a/Engine/Source/Runtime/Engine/Classes/AI/Navigation/NavFilters/RecastFilter_UseDefaultArea.h b/Engine/Source/Runtime/Engine/Classes/AI/Navigation/NavFilters/RecastFilter_UseDefaultArea.h
--- a/Engine/Source/Runtime/Engine/Classes/AI/Navigation/NavFilters/RecastFilter_UseDefaultArea.h
+++ b/Engine/Source/Runtime/Engine/Classes/AI/Navigation/NavFilters/RecastFilter_UseDefaultArea.h
@@ -7,9 +7,7 @@
 UCLASS(MinimalAPI)
 class URecastFilter_UseDefaultArea : public UNavigationQueryFilter
 {
-	GENERATED_BODY()
-public:
-	ENGINE_API URecastFilter_UseDefaultArea(const FObjectInitializer& ObjectInitializer = FObjectInitializer::Get());
+	GENERATED_UCLASS_BODY()
 
 	virtual void InitializeFilter(const ANavigationData* NavData, FNavigationQueryFilter* Filter) const override;
 };
diff --git a/Engine/Source/Runtime/Engine/Classes/AI/Navigation/NavLinkCustomComponent.h b/Engine/Source/Runtime/Engine/Classes/AI/Navigation/NavLinkCustomComponent.h
--- a/Engine/Source/Runtime/Engine/Classes/AI/Navigation/NavLinkCustomComponent.h
+++ b/Engine/Source/Runtime/Engine/Classes/AI/Navigation/NavLinkCustomComponent.h
@@ -24,9 +24,7 @@ struct FNavigationRelevantData;
 UCLASS()
 class ENGINE_API UNavLinkCustomComponent : public UNavRelevantComponent, public INavLinkCustomInterface
 {
-	GENERATED_BODY()
-public:
-	UNavLinkCustomComponent(const FObjectInitializer& ObjectInitializer = FObjectInitializer::Get());
+	GENERATED_UCLASS_BODY()
 
 	DECLARE_DELEGATE_ThreeParams(FOnMoveReachedLink, UNavLinkCustomComponent* /*ThisComp*/, UPathFollowingComponent* /*PathComp*/, const FVector& /*DestPoint*/);
 	DECLARE_DELEGATE_TwoParams(FBroadcastFilter, UNavLinkCustomComponent* /*ThisComp*/, TArray<UPathFollowingComponent*>& /*NotifyList*/);
diff --git a/Engine/Source/Runtime/Engine/Classes/AI/Navigation/NavLinkCustomInterface.h b/Engine/Source/Runtime/Engine/Classes/AI/Navigation/NavLinkCustomInterface.h
--- a/Engine/Source/Runtime/Engine/Classes/AI/Navigation/NavLinkCustomInterface.h
+++ b/Engine/Source/Runtime/Engine/Classes/AI/Navigation/NavLinkCustomInterface.h
@@ -24,15 +24,12 @@ struct FNavigationLink;
 UINTERFACE(MinimalAPI, meta=(CannotImplementInterfaceInBlueprint))
 class UNavLinkCustomInterface : public UInterface
 {
-	GENERATED_BODY()
-public:
-	ENGINE_API UNavLinkCustomInterface(const FObjectInitializer& ObjectInitializer = FObjectInitializer::Get());
+	GENERATED_UINTERFACE_BODY()
 };
 
 class ENGINE_API INavLinkCustomInterface
 {
-	GENERATED_BODY()
-public:
+	GENERATED_IINTERFACE_BODY()
 
 	/** Get basic link data: two points (relative to owner) and direction */
 	virtual void GetLinkData(FVector& LeftPt, FVector& RightPt, ENavLinkDirection::Type& Direction) const {};
diff --git a/Engine/Source/Runtime/Engine/Classes/AI/Navigation/NavLinkDefinition.h b/Engine/Source/Runtime/Engine/Classes/AI/Navigation/NavLinkDefinition.h
--- a/Engine/Source/Runtime/Engine/Classes/AI/Navigation/NavLinkDefinition.h
+++ b/Engine/Source/Runtime/Engine/Classes/AI/Navigation/NavLinkDefinition.h
@@ -19,7 +19,7 @@ namespace ENavLinkDirection
 USTRUCT()
 struct ENGINE_API FNavigationLinkBase
 {
-	GENERATED_BODY()
+	GENERATED_USTRUCT_BODY()
 
 	/** if greater than 0 nav system will attempt to project navlink's end point geometry below */
 	UPROPERTY(EditAnywhere, Category=Default, meta=(ClampMin = "0.0"))
@@ -100,7 +100,7 @@ struct ENGINE_API FNavigationLinkBase
 USTRUCT()
 struct ENGINE_API FNavigationLink : public FNavigationLinkBase
 {
-	GENERATED_BODY()
+	GENERATED_USTRUCT_BODY()
 
 	UPROPERTY(EditAnywhere, Category=Default, BlueprintReadWrite, meta=(MakeEditWidget=""))
 	FVector Left;
@@ -148,7 +148,7 @@ struct ENGINE_API FNavigationLink : public FNavigationLinkBase
 USTRUCT()
 struct ENGINE_API FNavigationSegmentLink : public FNavigationLinkBase
 {
-	GENERATED_BODY()
+	GENERATED_USTRUCT_BODY()
 
 	UPROPERTY(EditAnywhere, Category=Default, BlueprintReadWrite, meta=(MakeEditWidget=""))
 	FVector LeftStart;
@@ -197,9 +197,7 @@ struct ENGINE_API FNavigationSegmentLink : public FNavigationLinkBase
 UCLASS(abstract, Config=Engine, Blueprintable)
 class ENGINE_API UNavLinkDefinition : public UObject
 {
-	GENERATED_BODY()
-public:
-	UNavLinkDefinition(const FObjectInitializer& ObjectInitializer = FObjectInitializer::Get());
+	GENERATED_UCLASS_BODY()
 
 	UPROPERTY(EditAnywhere, Category=OffMeshLinks)
 	TArray<FNavigationLink> Links;
diff --git a/Engine/Source/Runtime/Engine/Classes/AI/Navigation/NavLinkHostInterface.h b/Engine/Source/Runtime/Engine/Classes/AI/Navigation/NavLinkHostInterface.h
--- a/Engine/Source/Runtime/Engine/Classes/AI/Navigation/NavLinkHostInterface.h
+++ b/Engine/Source/Runtime/Engine/Classes/AI/Navigation/NavLinkHostInterface.h
@@ -10,15 +10,12 @@ struct FNavigationSegmentLink;
 UINTERFACE(MinimalAPI, meta=(CannotImplementInterfaceInBlueprint))
 class UNavLinkHostInterface : public UInterface
 {
-	GENERATED_BODY()
-public:
-	ENGINE_API UNavLinkHostInterface(const FObjectInitializer& ObjectInitializer = FObjectInitializer::Get());
+	GENERATED_UINTERFACE_BODY()
 };
 
 class ENGINE_API INavLinkHostInterface
 {
-	GENERATED_BODY()
-public:
+	GENERATED_IINTERFACE_BODY()
 		
 	/**
 	 *	Retrieves UNavLinkDefinition derived UClasses hosted by this interface implementor
diff --git a/Engine/Source/Runtime/Engine/Classes/AI/Navigation/NavLinkProxy.h b/Engine/Source/Runtime/Engine/Classes/AI/Navigation/NavLinkProxy.h
--- a/Engine/Source/Runtime/Engine/Classes/AI/Navigation/NavLinkProxy.h
+++ b/Engine/Source/Runtime/Engine/Classes/AI/Navigation/NavLinkProxy.h
@@ -16,9 +16,7 @@ DECLARE_DYNAMIC_MULTICAST_DELEGATE_TwoParams( FSmartLinkReachedSignature, class
 UCLASS(Blueprintable)
 class ENGINE_API ANavLinkProxy : public AActor, public INavLinkHostInterface, public INavRelevantInterface
 {
-	GENERATED_BODY()
-public:
-	ANavLinkProxy(const FObjectInitializer& ObjectInitializer = FObjectInitializer::Get());
+	GENERATED_UCLASS_BODY()
 
 	/** Navigation links (point to point) added to navigation data */
 	UPROPERTY(EditAnywhere, Category=SimpleLink)
diff --git a/Engine/Source/Runtime/Engine/Classes/AI/Navigation/NavLinkRenderingComponent.h b/Engine/Source/Runtime/Engine/Classes/AI/Navigation/NavLinkRenderingComponent.h
--- a/Engine/Source/Runtime/Engine/Classes/AI/Navigation/NavLinkRenderingComponent.h
+++ b/Engine/Source/Runtime/Engine/Classes/AI/Navigation/NavLinkRenderingComponent.h
@@ -9,9 +9,7 @@ class UNavLinkDefinition;
 UCLASS(hidecategories=Object, editinlinenew)
 class ENGINE_API UNavLinkRenderingComponent : public UPrimitiveComponent
 {
-	GENERATED_BODY()
-public:
-	UNavLinkRenderingComponent(const FObjectInitializer& ObjectInitializer = FObjectInitializer::Get());
+	GENERATED_UCLASS_BODY()
 		
 	// Begin UPrimitiveComponent Interface
 	virtual FPrimitiveSceneProxy* CreateSceneProxy() override;
diff --git a/Engine/Source/Runtime/Engine/Classes/AI/Navigation/NavLinkTrivial.h b/Engine/Source/Runtime/Engine/Classes/AI/Navigation/NavLinkTrivial.h
--- a/Engine/Source/Runtime/Engine/Classes/AI/Navigation/NavLinkTrivial.h
+++ b/Engine/Source/Runtime/Engine/Classes/AI/Navigation/NavLinkTrivial.h
@@ -6,7 +6,5 @@
 UCLASS()
 class ENGINE_API UNavLinkTrivial : public UNavLinkDefinition
 {
-	GENERATED_BODY()
-public:
-	UNavLinkTrivial(const FObjectInitializer& ObjectInitializer = FObjectInitializer::Get());
-};
+	GENERATED_UCLASS_BODY()
+};
\ No newline at end of file
diff --git a/Engine/Source/Runtime/Engine/Classes/AI/Navigation/NavMeshBoundsVolume.h b/Engine/Source/Runtime/Engine/Classes/AI/Navigation/NavMeshBoundsVolume.h
--- a/Engine/Source/Runtime/Engine/Classes/AI/Navigation/NavMeshBoundsVolume.h
+++ b/Engine/Source/Runtime/Engine/Classes/AI/Navigation/NavMeshBoundsVolume.h
@@ -12,9 +12,7 @@
 UCLASS(MinimalAPI)
 class ANavMeshBoundsVolume : public AVolume
 {
-	GENERATED_BODY()
-public:
-	ENGINE_API ANavMeshBoundsVolume(const FObjectInitializer& ObjectInitializer = FObjectInitializer::Get());
+	GENERATED_UCLASS_BODY()
 
 	// Begin AActor Interface
 	virtual void PostRegisterAllComponents() override;
