REL - Prioritize Vendor over Quality setting in screenshot comparissions

diff --git a/Engine/Source/Runtime/AutomationMessages/Public/AutomationWorkerMessages.h b/Engine/Source/Runtime/AutomationMessages/Public/AutomationWorkerMessages.h
--- a/Engine/Source/Runtime/AutomationMessages/Public/AutomationWorkerMessages.h
+++ b/Engine/Source/Runtime/AutomationMessages/Public/AutomationWorkerMessages.h
@@ -530,15 +530,7 @@ public:
 			Score += 1000;
 		}
 
-		if (ResolutionQuality == OtherMetadata.ResolutionQuality && 
-			ViewDistanceQuality == OtherMetadata.ViewDistanceQuality &&
-			AntiAliasingQuality == OtherMetadata.AntiAliasingQuality &&
-			ShadowQuality == OtherMetadata.ShadowQuality &&
-			PostProcessQuality == OtherMetadata.PostProcessQuality &&
-			TextureQuality == OtherMetadata.TextureQuality &&
-			EffectsQuality == OtherMetadata.EffectsQuality &&
-			FoliageQuality == OtherMetadata.FoliageQuality
-			)
+		if (Vendor == OtherMetadata.Vendor)
 		{
 			Score += 100;
 		}
@@ -563,22 +555,57 @@ public:
 			Score += 10;
 		}
 
-		if (Vendor == OtherMetadata.Vendor)
+		if (AdapterName == OtherMetadata.AdapterName)
 		{
 			Score += 10;
 		}
 
-		if (AdapterName == OtherMetadata.AdapterName)
+		if (AdapterInternalDriverVersion == OtherMetadata.AdapterInternalDriverVersion)
 		{
 			Score += 10;
 		}
 
-		if (AdapterInternalDriverVersion == OtherMetadata.AdapterInternalDriverVersion)
+		if (AdapterUserDriverVersion == OtherMetadata.AdapterUserDriverVersion)
 		{
 			Score += 10;
 		}
 
-		if (AdapterUserDriverVersion == OtherMetadata.AdapterUserDriverVersion)
+		if (ResolutionQuality == OtherMetadata.ResolutionQuality)
+		{
+			Score += 10;
+		}
+
+		if (ViewDistanceQuality == OtherMetadata.ViewDistanceQuality)
+		{
+			Score += 10;
+		}
+
+		if (AntiAliasingQuality == OtherMetadata.AntiAliasingQuality)
+		{
+			Score += 10;
+		}
+
+		if (ShadowQuality == OtherMetadata.ShadowQuality)
+		{
+			Score += 10;
+		}
+
+		if (PostProcessQuality == OtherMetadata.PostProcessQuality)
+		{
+			Score += 10;
+		}
+
+		if (TextureQuality == OtherMetadata.TextureQuality)
+		{
+			Score += 10;
+		}
+
+		if (EffectsQuality == OtherMetadata.EffectsQuality)
+		{
+			Score += 10;
+		}
+
+		if (FoliageQuality == OtherMetadata.FoliageQuality)
 		{
 			Score += 10;
 		}
