#UE Fix initialization error

diff --git a/Engine/Source/Runtime/AutomationMessages/Public/AutomationWorkerMessages.h b/Engine/Source/Runtime/AutomationMessages/Public/AutomationWorkerMessages.h
--- a/Engine/Source/Runtime/AutomationMessages/Public/AutomationWorkerMessages.h
+++ b/Engine/Source/Runtime/AutomationMessages/Public/AutomationWorkerMessages.h
@@ -324,7 +324,7 @@ struct FAutomationWorkerRunTests : public FAutomationWorkerMessageBase
 	bool bPruneLogsOnSuccess;
 
 	/** Default constructor. */
-	FAutomationWorkerRunTests( ) :ExecutionCount(0), RoleIndex(0), bSendAnalytics(false) { }
+	FAutomationWorkerRunTests( ) :ExecutionCount(0), RoleIndex(0), bSendAnalytics(false), bPruneLogsOnSuccess(false) { }
 
 	/** Creates and initializes a new instance. */
 	FAutomationWorkerRunTests( uint32 InExecutionCount, int32 InRoleIndex, FString InTestName, FString InBeautifiedTestName, FString InFullTestPath, bool InSendAnalytics, bool InPruneLogsOnSuccess)
