Fix incorrect reference picking when a test has several alternative results.

The comparison metric now takes into account screenshot's Substrate metadata.

diff --git a/Engine/Source/Runtime/AutomationMessages/Public/AutomationWorkerMessages.h b/Engine/Source/Runtime/AutomationMessages/Public/AutomationWorkerMessages.h
--- a/Engine/Source/Runtime/AutomationMessages/Public/AutomationWorkerMessages.h
+++ b/Engine/Source/Runtime/AutomationMessages/Public/AutomationWorkerMessages.h
@@ -628,6 +628,11 @@ public:
 			Score += 1000;
 		}
 
+		if (bIsSubstrate == OtherMetadata.bIsSubstrate)
+		{
+			Score += 100;
+		}
+
 		if (Vendor == OtherMetadata.Vendor)
 		{
 			Score += 100;
